Fix: Missing include introduced in CL29731236

diff --git a/Engine/Plugins/Experimental/AVCodecs/VTCodecs/Source/VTCodecsRHI/Private/VTCodecsRHIModule.cpp b/Engine/Plugins/Experimental/AVCodecs/VTCodecs/Source/VTCodecsRHI/Private/VTCodecsRHIModule.cpp
--- a/Engine/Plugins/Experimental/AVCodecs/VTCodecs/Source/VTCodecsRHI/Private/VTCodecsRHIModule.cpp
+++ b/Engine/Plugins/Experimental/AVCodecs/VTCodecs/Source/VTCodecsRHI/Private/VTCodecsRHIModule.cpp
@@ -3,7 +3,7 @@
 #include "Modules/ModuleManager.h"
 
 #include "Misc/App.h"
-
+#include "RHI.h"
 #include "VT.h"
 #include "AVUtility.h"
 
