Fix: Build health issues from CL29731236

diff --git a/Engine/Plugins/Experimental/AVCodecs/VTCodecs/Source/VTCodecsRHI/Private/VTCodecsRHIModule.cpp b/Engine/Plugins/Experimental/AVCodecs/VTCodecs/Source/VTCodecsRHI/Private/VTCodecsRHIModule.cpp
--- a/Engine/Plugins/Experimental/AVCodecs/VTCodecs/Source/VTCodecsRHI/Private/VTCodecsRHIModule.cpp
+++ b/Engine/Plugins/Experimental/AVCodecs/VTCodecs/Source/VTCodecsRHI/Private/VTCodecsRHIModule.cpp
@@ -1,9 +1,10 @@
 // Copyright Epic Games, Inc. All Rights Reserved.
 
 #include "Modules/ModuleManager.h"
-
 #include "Misc/App.h"
+#include "Misc/CoreDelegates.h"
 #include "RHI.h"
+
 #include "VT.h"
 #include "AVUtility.h"
 
diff --git a/Engine/Plugins/Experimental/AVCodecs/VTCodecs/Source/VTCodecsRHI/VTCodecsRHI.build.cs b/Engine/Plugins/Experimental/AVCodecs/VTCodecs/Source/VTCodecsRHI/VTCodecsRHI.build.cs
--- a/Engine/Plugins/Experimental/AVCodecs/VTCodecs/Source/VTCodecsRHI/VTCodecsRHI.build.cs
+++ b/Engine/Plugins/Experimental/AVCodecs/VTCodecs/Source/VTCodecsRHI/VTCodecsRHI.build.cs
@@ -12,10 +12,11 @@ public class VTCodecsRHI : ModuleRules
         DefaultBuildSettings = BuildSettingsVersion.V2;
         
         PrivateDependencyModuleNames.AddRange(new string[] {
+            "Core",
             "Engine",
             "AVCodecsCore",
-            "VTCodecs",
-            "RHI"
+            "RHI",
+            "VTCodecs"
         });
 
         PublicDependencyModuleNames.AddRange(new string[] {
