Fix: Build health issues from CL29731236

diff --git a/Engine/Plugins/Experimental/AVCodecs/VTCodecs/Source/VTCodecs/Private/VTCodecsModule.cpp b/Engine/Plugins/Experimental/AVCodecs/VTCodecs/Source/VTCodecs/Private/VTCodecsModule.cpp
--- a/Engine/Plugins/Experimental/AVCodecs/VTCodecs/Source/VTCodecs/Private/VTCodecsModule.cpp
+++ b/Engine/Plugins/Experimental/AVCodecs/VTCodecs/Source/VTCodecs/Private/VTCodecsModule.cpp
@@ -1,5 +1,8 @@
 // Copyright Epic Games, Inc. All Rights Reserved.
 
+#include "Modules/ModuleManager.h"
+#include "Misc/App.h"
+
 #include "Video/Encoders/VideoEncoderVT.h"
 #include "Video/Encoders/Configs/VideoEncoderConfigH264.h"
 #include "Video/Encoders/Configs/VideoEncoderConfigH265.h"
diff --git a/Engine/Plugins/Experimental/AVCodecs/VTCodecs/Source/VTCodecs/Private/Video/Util/VTSessionHelpers.cpp b/Engine/Plugins/Experimental/AVCodecs/VTCodecs/Source/VTCodecs/Private/Video/Util/VTSessionHelpers.cpp
--- a/Engine/Plugins/Experimental/AVCodecs/VTCodecs/Source/VTCodecs/Private/Video/Util/VTSessionHelpers.cpp
+++ b/Engine/Plugins/Experimental/AVCodecs/VTCodecs/Source/VTCodecs/Private/Video/Util/VTSessionHelpers.cpp
@@ -3,8 +3,6 @@
 #include "Video/Util/VTSessionHelpers.h"
 #include "AVResult.h"
 
-#include <string>
-
 FString VTSessionHelpers::CFStringToString(const CFStringRef CfString) 
 {    
     // Get the size needed for UTF8 plus terminating character.
diff --git a/Engine/Plugins/Experimental/AVCodecs/VTCodecs/Source/VTCodecs/Public/Video/Util/VTSessionHelpers.h b/Engine/Plugins/Experimental/AVCodecs/VTCodecs/Source/VTCodecs/Public/Video/Util/VTSessionHelpers.h
--- a/Engine/Plugins/Experimental/AVCodecs/VTCodecs/Source/VTCodecs/Public/Video/Util/VTSessionHelpers.h
+++ b/Engine/Plugins/Experimental/AVCodecs/VTCodecs/Source/VTCodecs/Public/Video/Util/VTSessionHelpers.h
@@ -2,6 +2,9 @@
 
 #pragma once
 
+#include "CoreMinimal.h"
+#include "Containers/UnrealString.h"
+
 THIRD_PARTY_INCLUDES_START
 #include <VideoToolbox/VideoToolbox.h>
 THIRD_PARTY_INCLUDES_END
