Fix: Remove extra pixel data copy and fix bug when texture strides didn't match pixel buffer strides

diff --git a/Engine/Plugins/Experimental/AVCodecs/VTCodecs/Source/VTCodecs/Public/Video/Encoders/VideoEncoderVT.hpp b/Engine/Plugins/Experimental/AVCodecs/VTCodecs/Source/VTCodecs/Public/Video/Encoders/VideoEncoderVT.hpp
--- a/Engine/Plugins/Experimental/AVCodecs/VTCodecs/Source/VTCodecs/Public/Video/Encoders/VideoEncoderVT.hpp
+++ b/Engine/Plugins/Experimental/AVCodecs/VTCodecs/Source/VTCodecs/Public/Video/Encoders/VideoEncoderVT.hpp
@@ -52,6 +52,8 @@ void TVideoEncoderVT<TResource>::Close()
         CFRelease(Encoder);
         Encoder = nullptr;
     }
+    
+    bIsOpen = false;
 }
 
 template <typename TResource>
@@ -249,11 +251,8 @@ FAVResult TVideoEncoderVT<TResource>::SendFrame(TSharedPtr<FVideoResourceMetal>
             MTL::Texture* RawTexture = Resource->GetRaw();
             uint32_t Width = RawTexture->width();
             uint32_t Height = RawTexture->height();
-            TArray<uint8> PixelBytes;
-            PixelBytes.SetNumZeroed(4 * Width * Height);
 
-            Resource->GetRaw()->getBytes(PixelBytes.GetData(), 4 * Width, MTL::Region(0, 0, Width, Height), 0);
-            FMemory::BigBlockMemcpy(reinterpret_cast<uint8*>(CVPixelBufferGetBaseAddressOfPlane(PixelBuffer, 0)), PixelBytes.GetData(), PixelBytes.Num());
+            RawTexture->getBytes(reinterpret_cast<uint8*>(CVPixelBufferGetBaseAddressOfPlane(PixelBuffer, 0)), CVPixelBufferGetBytesPerRow(PixelBuffer), MTL::Region(0, 0, Width, Height), 0);
             CVPixelBufferUnlockBaseAddress(PixelBuffer, 0);
             
             CMTime PresentationTime = CMTimeMake(Timestamp, 1000);
@@ -266,7 +265,7 @@ FAVResult TVideoEncoderVT<TResource>::SendFrame(TSharedPtr<FVideoResourceMetal>
             }
 
             TUniquePtr<EncodeParams> Params = MakeUnique<EncodeParams>();
-            Params.Reset(new EncodeParams(this->GetAppliedConfig().Codec, PresentationTime));
+            Params.Reset(new EncodeParams(PendingConfig.Codec, PresentationTime));
 
             OSStatus Status = VTCompressionSessionEncodeFrame(Encoder, PixelBuffer, PresentationTime, kCMTimeInvalid, FrameProperties, (void*)Params.Release(), nullptr);
 
