Fix: Setting VTCodecs encoder target bitrate

diff --git a/Engine/Plugins/Experimental/AVCodecs/VTCodecs/Source/VTCodecs/Private/Video/Encoders/Configs/VideoEncoderConfigVT.cpp b/Engine/Plugins/Experimental/AVCodecs/VTCodecs/Source/VTCodecs/Private/Video/Encoders/Configs/VideoEncoderConfigVT.cpp
--- a/Engine/Plugins/Experimental/AVCodecs/VTCodecs/Source/VTCodecs/Private/Video/Encoders/Configs/VideoEncoderConfigVT.cpp
+++ b/Engine/Plugins/Experimental/AVCodecs/VTCodecs/Source/VTCodecs/Private/Video/Encoders/Configs/VideoEncoderConfigVT.cpp
@@ -22,11 +22,10 @@ TAVResult<OSType> FVideoEncoderConfigVT::ConvertFormat(EVideoFormat const& Forma
 		case EVideoFormat::NV12:
 			return kCVPixelFormatType_420YpCbCr8BiPlanarFullRange;
 		default:
-			return FAVResult(EAVResult::ErrorUnsupported, FString::Printf(TEXT("Pixel format %d is not supported"), Format), TEXT("NVENC"));
+			return FAVResult(EAVResult::ErrorUnsupported, FString::Printf(TEXT("Pixel format %d is not supported"), Format), TEXT("VT"));
 	}
 }
 
-
 template<>
 DLLEXPORT FAVResult FAVExtension::TransformConfig(FVideoEncoderConfigVT& OutConfig, FVideoEncoderConfig const& InConfig)
 {
@@ -37,6 +36,9 @@ DLLEXPORT FAVResult FAVExtension::TransformConfig(FVideoEncoderConfigVT& OutConf
 	OutConfig.TargetBitrate = InConfig.TargetBitrate > -1 ? InConfig.TargetBitrate : DEFAULT_BITRATE_TARGET;
 	OutConfig.MaxBitrate = InConfig.MaxBitrate > -1 ? InConfig.MaxBitrate : DEFAULT_BITRATE_MAX;
 	OutConfig.RateControlMode = InConfig.RateControlMode;
+    OutConfig.MinQP = InConfig.MinQP;
+    OutConfig.MaxQP = InConfig.MaxQP;
+	OutConfig.KeyframeInterval = InConfig.KeyframeInterval > 0 ? InConfig.KeyframeInterval : 0;
 
 	return EAVResult::Success;
 }
@@ -51,6 +53,9 @@ DLLEXPORT FAVResult FAVExtension::TransformConfig(FVideoEncoderConfig& OutConfig
 	OutConfig.TargetBitrate = InConfig.TargetBitrate > -1 ? InConfig.TargetBitrate : DEFAULT_BITRATE_TARGET;
 	OutConfig.MaxBitrate = InConfig.MaxBitrate > -1 ? InConfig.MaxBitrate : DEFAULT_BITRATE_MAX;
 	OutConfig.RateControlMode = InConfig.RateControlMode;
+    OutConfig.MinQP = InConfig.MinQP;
+    OutConfig.MaxQP = InConfig.MaxQP;
+	OutConfig.KeyframeInterval = InConfig.KeyframeInterval > 0 ? InConfig.KeyframeInterval : 0;
 
 	return EAVResult::Success;
 }
@@ -58,6 +63,30 @@ DLLEXPORT FAVResult FAVExtension::TransformConfig(FVideoEncoderConfig& OutConfig
 template <>
 DLLEXPORT FAVResult FAVExtension::TransformConfig(FVideoEncoderConfigVT& OutConfig, FVideoEncoderConfigH264 const& InConfig)
 {
+    OutConfig.Codec = kCMVideoCodecType_H264;
+
+    static auto const ConvertEntropyCodingMode = [](EH264EntropyCodingMode EntropyCodingMode) -> TAVResult<CFStringRef> {
+		switch (EntropyCodingMode)
+		{
+			case EH264EntropyCodingMode::Auto:
+				return nullptr;
+			case EH264EntropyCodingMode::CABAC:
+				return kVTH264EntropyMode_CABAC;
+			case EH264EntropyCodingMode::CAVLC:
+				return kVTH264EntropyMode_CAVLC;
+			default:
+				return FAVResult(EAVResult::ErrorUnsupported, FString::Printf(TEXT("H264 entropy coding mode %d is not supported"), EntropyCodingMode), TEXT("VT"));
+		}
+	};
+
+    TAVResult<CFStringRef> const ConvertedEntropyCodingMode = ConvertEntropyCodingMode(InConfig.EntropyCodingMode);
+	if (ConvertedEntropyCodingMode.IsNotSuccess())
+	{
+		return ConvertedEntropyCodingMode;
+	}
+
+    OutConfig.EntropyCodingMode = ConvertedEntropyCodingMode;
+
 	static auto const ConvertProfile = [](EH264Profile Profile) -> TAVResult<CFStringRef>
 	{
 		switch (Profile)
@@ -71,17 +100,15 @@ DLLEXPORT FAVResult FAVExtension::TransformConfig(FVideoEncoderConfigVT& OutConf
 		case EH264Profile::ConstrainedHigh:
 			return kVTProfileLevel_H264_ConstrainedHigh_AutoLevel;
 		default:
-			return FAVResult(EAVResult::ErrorUnsupported, FString::Printf(TEXT("H264 profile %d is not supported"), Profile), TEXT("AMF"));
+			return FAVResult(EAVResult::ErrorUnsupported, FString::Printf(TEXT("H264 profile %d is not supported"), Profile), TEXT("VT"));
 		}
 	};
 
-	OutConfig.Codec = kCMVideoCodecType_H264;
-
 	TAVResult<CFStringRef> const ConvertedProfile = ConvertProfile(InConfig.Profile);
 	if (ConvertedProfile.IsNotSuccess())
 	{
 		return ConvertedProfile;
-	}
+    }
 
 	OutConfig.Profile = ConvertedProfile.ReturnValue;
 
@@ -102,7 +129,7 @@ DLLEXPORT FAVResult FAVExtension::TransformConfig(FVideoEncoderConfigVT& OutConf
         case EH265Profile::Main422_10:
             return kVTProfileLevel_HEVC_Main42210_AutoLevel;
 		default:
-			return FAVResult(EAVResult::ErrorUnsupported, FString::Printf(TEXT("H265 profile %d is not supported"), Profile), TEXT("AMF"));
+			return FAVResult(EAVResult::ErrorUnsupported, FString::Printf(TEXT("H265 profile %d is not supported"), Profile), TEXT("VT"));
 		}
 	};
 
diff --git a/Engine/Plugins/Experimental/AVCodecs/VTCodecs/Source/VTCodecs/Public/Video/Encoders/Configs/VideoEncoderConfigVT.h b/Engine/Plugins/Experimental/AVCodecs/VTCodecs/Source/VTCodecs/Public/Video/Encoders/Configs/VideoEncoderConfigVT.h
--- a/Engine/Plugins/Experimental/AVCodecs/VTCodecs/Source/VTCodecs/Public/Video/Encoders/Configs/VideoEncoderConfigVT.h
+++ b/Engine/Plugins/Experimental/AVCodecs/VTCodecs/Source/VTCodecs/Public/Video/Encoders/Configs/VideoEncoderConfigVT.h
@@ -27,6 +27,9 @@ public:
 	CMVideoCodecType Codec;
 	CFStringRef Profile;
     EVideoFormat PixelFormat;
+    CFStringRef EntropyCodingMode;
+    int32 MinQP;
+    int32 MaxQP;
 
 	FVideoEncoderConfigVT()
 		: FAVConfig()
@@ -60,7 +63,10 @@ public:
 				this->KeyframeInterval == Other.KeyframeInterval &&
 				this->Preset == Other.Preset &&
 				this->TargetBitrate == Other.TargetBitrate &&
-                this->PixelFormat == Other.PixelFormat;
+                this->PixelFormat == Other.PixelFormat &&
+                this->EntropyCodingMode == Other.EntropyCodingMode &&
+                this->MinQP == Other.MinQP &&
+                this->MaxQP == Other.MaxQP;
 	}
 
 	bool operator!=(FVideoEncoderConfigVT const& Other) const
diff --git a/Engine/Plugins/Experimental/AVCodecs/VTCodecs/Source/VTCodecs/Public/Video/Encoders/VideoEncoderVT.hpp b/Engine/Plugins/Experimental/AVCodecs/VTCodecs/Source/VTCodecs/Public/Video/Encoders/VideoEncoderVT.hpp
--- a/Engine/Plugins/Experimental/AVCodecs/VTCodecs/Source/VTCodecs/Public/Video/Encoders/VideoEncoderVT.hpp
+++ b/Engine/Plugins/Experimental/AVCodecs/VTCodecs/Source/VTCodecs/Public/Video/Encoders/VideoEncoderVT.hpp
@@ -74,7 +74,14 @@ FAVResult TVideoEncoderVT<TResource>::ApplyConfig()
 			{
                 if(this->AppliedConfig.Width == PendingConfig.Width
                     && this->AppliedConfig.Height == PendingConfig.Height
-                    && this->AppliedConfig.Codec == PendingConfig.Codec)
+                    && this->AppliedConfig.Codec == PendingConfig.Codec
+                    && this->AppliedConfig.RateControlMode == PendingConfig.RateControlMode
+                    && this->AppliedConfig.KeyframeInterval == PendingConfig.KeyframeInterval
+                    && this->AppliedConfig.Profile == PendingConfig.Profile
+                    && this->AppliedConfig.PixelFormat == PendingConfig.PixelFormat
+                    && this->AppliedConfig.EntropyCodingMode == PendingConfig.EntropyCodingMode
+                    && this->AppliedConfig.MinQP == PendingConfig.MinQP
+                    && this->AppliedConfig.MaxQP == PendingConfig.MaxQP)
                 {
                     // Only bitrates can be configured on the fly
                     SetEncoderBitrate(PendingConfig);
@@ -113,25 +120,25 @@ FAVResult TVideoEncoderVT<TResource>::ApplyConfig()
                 CONDITIONAL_RELEASE(IOSurfaceValue);
                 CONDITIONAL_RELEASE(PixelFormat);
 
-                CFMutableDictionaryRef EncoderSpecification = CFDictionaryCreateMutable(kCFAllocatorDefault, 2, &kCFTypeDictionaryKeyCallBacks, &kCFTypeDictionaryValueCallBacks);
+                CFMutableDictionaryRef EncoderSpecification = CFDictionaryCreateMutable(kCFAllocatorDefault, 1, &kCFTypeDictionaryKeyCallBacks, &kCFTypeDictionaryValueCallBacks);
                 
-                if(PendingConfig.Codec == kCMVideoCodecType_H264)
-                {
-                    // TODO (belchy06): Not all encoders are real-time. This should be moved to the config
-                    CFDictionarySetValue(EncoderSpecification, kVTVideoEncoderSpecification_EnableLowLatencyRateControl, kCFBooleanTrue);
-                }
-                else if(PendingConfig.Codec == kCMVideoCodecType_HEVC)
-                {
-                    CFDictionarySetValue(EncoderSpecification, kVTVideoEncoderSpecification_EnableHardwareAcceleratedVideoEncoder, kCFBooleanTrue);
-                }
+                // We explicitly state here that we want a hardware encoder. If it's optional, replace with kVTVideoEncoderSpecification_EnableHardwareAcceleratedVideoEncoder
+                CFDictionarySetValue(EncoderSpecification, kVTVideoEncoderSpecification_RequireHardwareAcceleratedVideoEncoder, kCFBooleanTrue);
                 
-                OSStatus Result = VTCompressionSessionCreate(kCFAllocatorDefault, PendingConfig.Width, PendingConfig.Height, PendingConfig.Codec, EncoderSpecification, SourceAttributes, NULL /* Default Compressed Data Allocator */, Internal::VTCompressionOutputCallback, this, &Encoder);
+                OSStatus Result = VTCompressionSessionCreate(kCFAllocatorDefault, PendingConfig.Width, PendingConfig.Height, PendingConfig.Codec, EncoderSpecification, SourceAttributes, nullptr /* Default Compressed Data Allocator */, Internal::VTCompressionOutputCallback, this, &Encoder);
                     
                 if(Result != 0)
                 {
-                    FAVResult::Log(EAVResult::ErrorCreating, TEXT("Failed to create VTCompressionSession"), TEXT("VT"), Result);
+                    CONDITIONAL_RELEASE(SourceAttributes);
+                    CONDITIONAL_RELEASE(EncoderSpecification);
+                    return FAVResult(EAVResult::ErrorCreating, TEXT("Failed to create VTCompressionSession"), TEXT("VT"), Result);
                 }
 
+                CFBooleanRef bIsUsingHardwareEncoder;
+                Result = VTSessionCopyProperty(Encoder, kVTCompressionPropertyKey_UsingHardwareAcceleratedVideoEncoder, NULL, &bIsUsingHardwareEncoder);
+                FAVResult::Log(EAVResult::Unknown, FString::Printf(TEXT("Created compression session. UsingHardwareEncoder: %s"), (Result == 0 && CFBooleanGetValue(bIsUsingHardwareEncoder) ? TEXT("TRUE") : TEXT("FALSE"))), TEXT("VT"));
+
+                CONDITIONAL_RELEASE(bIsUsingHardwareEncoder);
                 CONDITIONAL_RELEASE(SourceAttributes);
                 CONDITIONAL_RELEASE(EncoderSpecification);
 
@@ -147,22 +154,28 @@ FAVResult TVideoEncoderVT<TResource>::ApplyConfig()
 
 template <typename TResource>
 FAVResult TVideoEncoderVT<TResource>::ConfigureCompressionSession(FVideoEncoderConfigVT const& Config)  
-{
+{   
     VTSessionHelpers::SetVTSessionProperty(Encoder, kVTCompressionPropertyKey_RealTime, true);
     
-    if(Config.Codec == kCMVideoCodecType_H264)
-    {
-        VTSessionHelpers::SetVTSessionProperty(Encoder, kVTCompressionPropertyKey_ProfileLevel, Config.Profile);
-    }
+    VTSessionHelpers::SetVTSessionProperty(Encoder, kVTCompressionPropertyKey_ProfileLevel, Config.Profile);
     
     VTSessionHelpers::SetVTSessionProperty(Encoder, kVTCompressionPropertyKey_AllowFrameReordering, false);
-
-    VTSessionHelpers::SetVTSessionProperty(Encoder, kVTCompressionPropertyKey_MaxKeyFrameInterval, Config.KeyframeInterval);
-
-    VTSessionHelpers::SetVTSessionProperty(Encoder, kVTCompressionPropertyKey_ExpectedFrameRate, Config.FrameRate);
-
+    
     SetEncoderBitrate(Config);
     
+    VTSessionHelpers::SetVTSessionProperty(Encoder, kVTCompressionPropertyKey_MaxKeyFrameInterval, Config.KeyframeInterval);
+    
+    VTSessionHelpers::SetVTSessionProperty(Encoder, kVTCompressionPropertyKey_MaxKeyFrameIntervalDuration, Config.KeyframeInterval * Config.FrameRate);
+    
+    VTSessionHelpers::SetVTSessionProperty(Encoder, kVTCompressionPropertyKey_MinAllowedFrameQP, Config.MinQP);
+    
+    VTSessionHelpers::SetVTSessionProperty(Encoder, kVTCompressionPropertyKey_MaxAllowedFrameQP, Config.MaxQP);
+    
+    if(Config.Codec == kCMVideoCodecType_H264)
+    {
+        VTSessionHelpers::SetVTSessionProperty(Encoder, kVTCompressionPropertyKey_H264EntropyMode, Config.EntropyCodingMode);
+    }
+    
     OSStatus Status = VTCompressionSessionPrepareToEncodeFrames(Encoder);
     if(Status != 0)
     {
@@ -175,37 +188,29 @@ FAVResult TVideoEncoderVT<TResource>::ConfigureCompressionSession(FVideoEncoderC
 template <typename TResource>
 FAVResult TVideoEncoderVT<TResource>::SetEncoderBitrate(FVideoEncoderConfigVT const& Config)
 {
-    // CBR not supported by encoders in their current configuration
-    // if(Config.RateControlMode == ERateControlMode::CBR)
-    // {
-    //     VTSessionHelpers::SetVTSessionProperty(Encoder, kVTCompressionPropertyKey_ConstantBitRate, Config.TargetBitrate);
-    // }
-    // else
-    {
-        VTSessionHelpers::SetVTSessionProperty(Encoder, kVTCompressionPropertyKey_AverageBitRate, Config.TargetBitrate);
+    VTSessionHelpers::SetVTSessionProperty(Encoder, kVTCompressionPropertyKey_AverageBitRate, Config.TargetBitrate);
 
-        int64_t DataLimitBytesPerSecondValue = static_cast<int64_t>(Config.MaxBitrate * 1.5f / 8);
-
-        CFNumberRef BytesPerSecond = CFNumberCreate(kCFAllocatorDefault, kCFNumberSInt64Type, &DataLimitBytesPerSecondValue);
-        
-        int64_t OneSecondValue = 1;
-        CFNumberRef OneSecond = CFNumberCreate(kCFAllocatorDefault, kCFNumberSInt64Type, &OneSecondValue);
+    // Bits to bytes conversion
+    int64_t BytesPerSecondValue = static_cast<int64_t>(Config.MaxBitrate >> 3);
+    CFNumberRef BytesPerSecond = CFNumberCreate(kCFAllocatorDefault, kCFNumberSInt64Type, &BytesPerSecondValue);
     
-        const void* Nums[2] = { BytesPerSecond, OneSecond };
+    int64_t OneSecondValue = 1;
+    CFNumberRef OneSecond = CFNumberCreate(kCFAllocatorDefault, kCFNumberSInt64Type, &OneSecondValue);
     
-        CFArrayRef DataRateLimits = CFArrayCreate(nullptr, Nums, 2, &kCFTypeArrayCallBacks);
+    const void* Nums[2] = { BytesPerSecond, OneSecond };
+    CFArrayRef DataRateLimits = CFArrayCreate(nullptr, Nums, 2, &kCFTypeArrayCallBacks);
     
-        OSStatus Result = VTSessionSetProperty(Encoder, kVTCompressionPropertyKey_DataRateLimits, DataRateLimits);
-
-        if(Result != 0)
-        {
-            FAVResult::Log(EAVResult::Error, TEXT("Failed to set kVTCompressionPropertyKey_DataRateLimits"), TEXT("VT"), Result);
-        }
-
-        CONDITIONAL_RELEASE(BytesPerSecond);
-        CONDITIONAL_RELEASE(OneSecond);
-        CONDITIONAL_RELEASE(DataRateLimits);
+    OSStatus Result = VTSessionSetProperty(Encoder, kVTCompressionPropertyKey_DataRateLimits, DataRateLimits);
+    
+    CONDITIONAL_RELEASE(BytesPerSecond);
+    CONDITIONAL_RELEASE(OneSecond);
+    CONDITIONAL_RELEASE(DataRateLimits);
+    
+    if(Result != 0)
+    {
+        return FAVResult(EAVResult::Error, TEXT("Failed to set kVTCompressionPropertyKey_DataRateLimits"), TEXT("VT"), Result);
     }
+    
 
     return EAVResult::Success;
 }
