Fix: VideoToolbox hardware decode

diff --git a/Engine/Plugins/Experimental/AVCodecs/AVCodecsCore/Source/AVCodecsCore/Public/Video/Resources/Metal/VideoResourceMetal.h b/Engine/Plugins/Experimental/AVCodecs/AVCodecsCore/Source/AVCodecsCore/Public/Video/Resources/Metal/VideoResourceMetal.h
--- a/Engine/Plugins/Experimental/AVCodecs/AVCodecsCore/Source/AVCodecsCore/Public/Video/Resources/Metal/VideoResourceMetal.h
+++ b/Engine/Plugins/Experimental/AVCodecs/AVCodecsCore/Source/AVCodecsCore/Public/Video/Resources/Metal/VideoResourceMetal.h
@@ -8,6 +8,7 @@
 
 #include "AVContext.h"
 #include "Video/VideoResource.h"
+#include "Containers/ResourceArray.h"
 
 THIRD_PARTY_INCLUDES_START
 #include "MetalInclude.h"
@@ -43,6 +44,43 @@ public:
 	virtual FAVResult Validate() const override;
 };
 
+/**
+ * Passes a CVPixelBufferRef through to the RHI to wrap in an RHI texture without traversing system memory.
+ */
+class AVCODECSCORE_API FBulkDataMetal : public FResourceBulkDataInterface
+{
+public:
+    FBulkDataMetal(CFTypeRef InImageBuffer)
+        : ImageBuffer(InImageBuffer)
+    {
+        check(ImageBuffer);
+        CFRetain(ImageBuffer);
+    }
+    virtual ~FBulkDataMetal()
+    {
+        CFRelease(ImageBuffer);
+        ImageBuffer = nullptr;
+    }
+public:
+    virtual void Discard() override
+    {
+        delete this;
+    }
+    virtual const void* GetResourceBulkData() const override
+    {
+        return ImageBuffer;
+    }
+    virtual uint32 GetResourceBulkDataSize() const override
+    {
+        return ImageBuffer ? ~0u : 0;
+    }
+    virtual EBulkDataType GetResourceType() const override
+    {
+        return EBulkDataType::MediaTexture;
+    }
+    CFTypeRef ImageBuffer;
+};
+
 DECLARE_TYPEID(FVideoContextMetal, AVCODECSCORE_API);
 DECLARE_TYPEID(FVideoResourceMetal, AVCODECSCORE_API);
 
diff --git a/Engine/Plugins/Experimental/AVCodecs/AVCodecsCore/Source/AVCodecsCore/Public/Video/VideoResource.h b/Engine/Plugins/Experimental/AVCodecs/AVCodecsCore/Source/AVCodecsCore/Public/Video/VideoResource.h
--- a/Engine/Plugins/Experimental/AVCodecs/AVCodecsCore/Source/AVCodecsCore/Public/Video/VideoResource.h
+++ b/Engine/Plugins/Experimental/AVCodecs/AVCodecsCore/Source/AVCodecsCore/Public/Video/VideoResource.h
@@ -6,6 +6,7 @@
 #include "AVResource.h"
 #include "AVUtility.h"
 #include "PixelFormat.h"
+#include "Containers/ResourceArray.h"
 
 enum class EVideoFormat : uint8
 {
@@ -62,6 +63,11 @@ public:
 	 * Height in pixels.
 	 */
 	uint32 Height;
+    
+    /**
+     * An implementation of the BulkData interface. Used for creating textures from existing data
+     */
+    FResourceBulkDataInterface* BulkData;
 
 	/**
 	 * If this is a descriptor stored in a different format than it really is
@@ -74,18 +80,19 @@ public:
 		: Format(Descriptor.Format)
 		, Width(Descriptor.Width)
 		, Height(Descriptor.Height)
+        , BulkData(Descriptor.BulkData)
 		, RawDescriptor(Descriptor.RawDescriptor ? new FVideoDescriptor(*Descriptor.RawDescriptor) : nullptr)
 	{
 
 	}
 
-	FVideoDescriptor(EVideoFormat Format, uint32 Width, uint32 Height)
-		: Format(Format), Width(Width), Height(Height), RawDescriptor(nullptr)
+	FVideoDescriptor(EVideoFormat Format, uint32 Width, uint32 Height, FResourceBulkDataInterface* BulkData = nullptr)
+		: Format(Format), Width(Width), Height(Height), BulkData(BulkData), RawDescriptor(nullptr)
 	{
 	}
 
-	FVideoDescriptor(EVideoFormat Format, uint32 Width, uint32 Height, const FVideoDescriptor& RawDescriptor)
-		: Format(Format), Width(Width), Height(Height), RawDescriptor(new FVideoDescriptor(RawDescriptor))
+	FVideoDescriptor(EVideoFormat Format, uint32 Width, uint32 Height, const FVideoDescriptor& RawDescriptor, FResourceBulkDataInterface* BulkData = nullptr)
+		: Format(Format), Width(Width), Height(Height), BulkData(BulkData), RawDescriptor(new FVideoDescriptor(RawDescriptor))
 	{
 
 	}
diff --git a/Engine/Plugins/Experimental/AVCodecs/AVCodecsCore/Source/AVCodecsCoreRHI/Private/Video/Resources/VideoResourceRHI.cpp b/Engine/Plugins/Experimental/AVCodecs/AVCodecsCore/Source/AVCodecsCoreRHI/Private/Video/Resources/VideoResourceRHI.cpp
--- a/Engine/Plugins/Experimental/AVCodecs/AVCodecsCore/Source/AVCodecsCoreRHI/Private/Video/Resources/VideoResourceRHI.cpp
+++ b/Engine/Plugins/Experimental/AVCodecs/AVCodecsCore/Source/AVCodecsCoreRHI/Private/Video/Resources/VideoResourceRHI.cpp
@@ -162,6 +162,11 @@ TSharedPtr<FVideoResourceRHI> FVideoResourceRHI::Create(TSharedPtr<FAVDevice> co
 		}
 
 		TextureDesc.AddFlags(AdditionalFlags);
+        
+        if (Descriptor.BulkData)
+        {
+            TextureDesc.SetBulkData(Descriptor.BulkData);
+        }
 
 		// Unreals support for NV12 asnd P010 is not universally supported so we actually use R8 or G16 under the hood
 		switch (Descriptor.Format)
@@ -608,7 +613,7 @@ DLLEXPORT FAVResult FAVExtension::TransformResource(TSharedPtr<FVideoResourceMet
             const FRHITextureCreateDesc Desc =
                 FRHITextureCreateDesc::Create2D(TEXT("FAVExtension::TransformResource"), Descriptor.Width, Descriptor.Height, PF_B8G8R8A8)
                 .SetFlags(ETextureCreateFlags::SRGB | ETextureCreateFlags::Dynamic | ETextureCreateFlags::NoTiling | ETextureCreateFlags::ShaderResource)
-                .SetBulkData(new FCVBulkData(TextureRef));
+                .SetBulkData(new FBulkDataMetal(TextureRef));
                         
             ENQUEUE_RENDER_COMMAND(FAVExtensionTransformResource)(
                 [&Source = InResource->GetRaw().Texture, &Desc, &Fence, &Destination](FRHICommandListImmediate& RHICmdList)
diff --git a/Engine/Plugins/Experimental/AVCodecs/AVCodecsCore/Source/AVCodecsCoreRHI/Public/Video/Resources/VideoResourceRHI.h b/Engine/Plugins/Experimental/AVCodecs/AVCodecsCore/Source/AVCodecsCoreRHI/Public/Video/Resources/VideoResourceRHI.h
--- a/Engine/Plugins/Experimental/AVCodecs/AVCodecsCore/Source/AVCodecsCoreRHI/Public/Video/Resources/VideoResourceRHI.h
+++ b/Engine/Plugins/Experimental/AVCodecs/AVCodecsCore/Source/AVCodecsCoreRHI/Public/Video/Resources/VideoResourceRHI.h
@@ -5,7 +5,6 @@
 #include "AVContext.h"
 #include "AVExtension.h"
 #include "Video/VideoResource.h"
-#include "Containers/ResourceArray.h"
 #include "RHI.h"
 
 #if AVCODECS_USE_D3D
@@ -73,45 +72,6 @@ protected:
 	virtual TSharedPtr<FVideoResourceRHI> TryResolve(TSharedPtr<FAVDevice> const& Device, FVideoDescriptor const& Descriptor) override;
 };
 
-#if AVCODECS_USE_METAL
-/**
- * Passes a CVMetalTextureRef or CVPixelBufferRef through to the RHI to wrap in an RHI texture without traversing system memory.
- */
-class AVCODECSCORERHI_API FCVBulkData : public FResourceBulkDataInterface
-{
-public:
-    FCVBulkData(CFTypeRef InImageBuffer)
-        : ImageBuffer(InImageBuffer)
-    {
-        check(ImageBuffer);
-        CFRetain(ImageBuffer);
-    }
-    virtual ~FCVBulkData()
-    {
-        CFRelease(ImageBuffer);
-        ImageBuffer = nullptr;
-    }
-public:
-    virtual void Discard() override
-    {
-        delete this;
-    }
-    virtual const void* GetResourceBulkData() const override
-    {
-        return ImageBuffer;
-    }
-    virtual uint32 GetResourceBulkDataSize() const override
-    {
-        return ImageBuffer ? ~0u : 0;
-    }
-    virtual EBulkDataType GetResourceType() const override
-    {
-        return EBulkDataType::MediaTexture;
-    }
-    CFTypeRef ImageBuffer;
-};
-#endif
-
 /*
 * Handles pixel format or colorspace transformations
 */
diff --git a/Engine/Plugins/Experimental/AVCodecs/VTCodecs/Source/VTCodecs/Public/Video/Decoders/Configs/VideoDecoderConfigVT.h b/Engine/Plugins/Experimental/AVCodecs/VTCodecs/Source/VTCodecs/Public/Video/Decoders/Configs/VideoDecoderConfigVT.h
--- a/Engine/Plugins/Experimental/AVCodecs/VTCodecs/Source/VTCodecs/Public/Video/Decoders/Configs/VideoDecoderConfigVT.h
+++ b/Engine/Plugins/Experimental/AVCodecs/VTCodecs/Source/VTCodecs/Public/Video/Decoders/Configs/VideoDecoderConfigVT.h
@@ -17,6 +17,7 @@ public:
 	
 public:
 	CMVideoCodecType Codec;
+    CMVideoFormatDescriptionRef VideoFormat;
 
 	FVideoDecoderConfigVT()
 		: FAVConfig()
@@ -28,6 +29,26 @@ public:
 	{
 		*this = Other;
 	}
+    
+    void SetVideoFormat(CMVideoFormatDescriptionRef Format)
+    {
+        if(VideoFormat == Format)
+        {
+            return;
+        }
+
+        if (VideoFormat)
+        {
+            CFRelease(VideoFormat);
+            VideoFormat = nullptr;
+        }
+
+        VideoFormat = Format;
+        if (VideoFormat)
+        {
+            CFRetain(VideoFormat);
+        }
+    }
 
 	FVideoDecoderConfigVT& operator=(FVideoDecoderConfigVT const& Other)
 	{
@@ -38,7 +59,8 @@ public:
 
 	bool operator==(FVideoDecoderConfigVT const& Other) const
 	{
-        return this->Codec == Other.Codec;
+        return this->Codec == Other.Codec &&
+               CMFormatDescriptionEqual(this->VideoFormat, Other.VideoFormat);
 	}
 
 	bool operator!=(FVideoDecoderConfigVT const& Other) const
diff --git a/Engine/Plugins/Experimental/AVCodecs/VTCodecs/Source/VTCodecs/Public/Video/Decoders/VideoDecoderVT.h b/Engine/Plugins/Experimental/AVCodecs/VTCodecs/Source/VTCodecs/Public/Video/Decoders/VideoDecoderVT.h
--- a/Engine/Plugins/Experimental/AVCodecs/VTCodecs/Source/VTCodecs/Public/Video/Decoders/VideoDecoderVT.h
+++ b/Engine/Plugins/Experimental/AVCodecs/VTCodecs/Source/VTCodecs/Public/Video/Decoders/VideoDecoderVT.h
@@ -20,8 +20,8 @@ template <typename TResource>
 class TVideoDecoderVT : public TVideoDecoder<TResource, FVideoDecoderConfigVT>
 {
 private:
-	VTDecompressionSessionRef Decoder = nullptr;
-    CMVideoFormatDescriptionRef VideoFormat;
+    VTDecompressionSessionRef Decoder = nullptr;
+    CVMetalTextureCacheRef TextureCache = nullptr;
     CMMemoryPoolRef MemoryPool;
 
 	uint8 bInitialized : 1;
@@ -76,7 +76,6 @@ public:
 	virtual bool IsOpen() const override;
 	virtual FAVResult Open(TSharedRef<FAVDevice> const& NewDevice, TSharedRef<FAVInstance> const& NewInstance) override;
 	virtual void Close() override;
-    void ResetDecompressionSession();
     void DestroyDecompressionSession();
     void ConfigureDecompressionSession();
 
@@ -88,8 +87,6 @@ public:
     FAVResult HandleFrame(void* Params, OSStatus Status, VTDecodeInfoFlags InfoFlags, CVImageBufferRef ImageBuffer, CMTime Timestamp, CMTime Duration);
 
 	virtual FAVResult SendPacket(FVideoPacket const& Packet) override;
-
-    void SetVideoFormat(CMVideoFormatDescriptionRef Format);
 };
 
 namespace Internal 
diff --git a/Engine/Plugins/Experimental/AVCodecs/VTCodecs/Source/VTCodecs/Public/Video/Decoders/VideoDecoderVT.hpp b/Engine/Plugins/Experimental/AVCodecs/VTCodecs/Source/VTCodecs/Public/Video/Decoders/VideoDecoderVT.hpp
--- a/Engine/Plugins/Experimental/AVCodecs/VTCodecs/Source/VTCodecs/Public/Video/Decoders/VideoDecoderVT.hpp
+++ b/Engine/Plugins/Experimental/AVCodecs/VTCodecs/Source/VTCodecs/Public/Video/Decoders/VideoDecoderVT.hpp
@@ -49,7 +49,6 @@ template <typename TResource>
 void TVideoDecoderVT<TResource>::Close()
 {
     DestroyDecompressionSession();
-    SetVideoFormat(nullptr);
 
     if(MemoryPool)
     {
@@ -60,63 +59,6 @@ void TVideoDecoderVT<TResource>::Close()
     bIsOpen = false;
 }
 
-template <typename TResource>
-void TVideoDecoderVT<TResource>::ResetDecompressionSession()
-{
-    DestroyDecompressionSession();
-
-    if(!VideoFormat)
-    {
-        FAVResult::Log(EAVResult::PendingInput, TEXT("Waiting for VideoFormat"), TEXT("VT"));                   
-        return; 
-    }
-    
-    // Set source image buffer attributes. These attributes will be present on
-    // buffers retrieved from the decoder's pixel buffer pool.
-    const size_t AttributesSize = 3;
-    CFTypeRef Keys[AttributesSize] = 
-    {
-        kCVPixelBufferOpenGLCompatibilityKey,
-        kCVPixelBufferIOSurfacePropertiesKey,
-        kCVPixelBufferPixelFormatTypeKey
-    };
-
-    CFDictionaryRef IOSurfaceValue = CFDictionaryCreate(kCFAllocatorDefault, nullptr, nullptr, 0, &kCFTypeDictionaryKeyCallBacks, &kCFTypeDictionaryValueCallBacks);
-    // TODO (belchy06): This should support more than ARGB8 pixel format
-    int64_t PixelType = kCVPixelFormatType_32BGRA;
-    CFNumberRef PixelFormat = CFNumberCreate(nullptr, kCFNumberLongType, &PixelType);
-
-    CFTypeRef Values[AttributesSize] =
-    { 
-        kCFBooleanTrue, 
-        IOSurfaceValue, 
-        PixelFormat 
-    };
-
-    CFDictionaryRef Attributes = CFDictionaryCreate(kCFAllocatorDefault, Keys, Values, AttributesSize, &kCFTypeDictionaryKeyCallBacks, &kCFTypeDictionaryValueCallBacks);
-
-    CONDITIONAL_RELEASE(IOSurfaceValue);
-    CONDITIONAL_RELEASE(PixelFormat);
-
-    VTDecompressionOutputCallbackRecord Record = 
-    {
-        Internal::VTDecompressionOutputCallback, 
-        this
-    };
-                
-    OSStatus Result = VTDecompressionSessionCreate(kCFAllocatorDefault, VideoFormat, nullptr, Attributes, &Record, &Decoder);
-
-    if(Result != 0)
-    {
-        DestroyDecompressionSession();
-        FAVResult::Log(EAVResult::ErrorCreating, TEXT("Failed to create VTDecompressionSession"), TEXT("VT"), Result);
-    }
-
-    CONDITIONAL_RELEASE(Attributes);
-
-    ConfigureDecompressionSession();
-}
-
 template <typename TResource>
 void TVideoDecoderVT<TResource>::DestroyDecompressionSession()
 {
@@ -150,23 +92,48 @@ FAVResult TVideoDecoderVT<TResource>::ApplyConfig()
 		{
 			if (IsInitialized())
 			{
-                if(this->AppliedConfig.Codec == PendingConfig.Codec)
-                {
-                    // TODO (belchy06): Reconfiguration
-                }
-                else
+                // VideoToolbox decoder doesn't support reconfiguration. If any aspect of the config changes,
+                // the entire session must be re-created
+                if (Decoder)
                 {
-                    if (Decoder) 
-                    {
-                        DestroyDecompressionSession();
-                        FAVResult::Log(EAVResult::Success, TEXT("Re-initializing decoding session"), TEXT("VT"));
-                    }
+                    DestroyDecompressionSession();
+                    FAVResult::Log(EAVResult::Success, TEXT("Re-initializing decoding session"), TEXT("VT"));
                 }
 			}
 
 			if (!IsInitialized())
 			{
-                ResetDecompressionSession();
+                // Set source image buffer attributes. These attributes will be present on
+                // buffers retrieved from the decoder's pixel buffer pool.
+                CFMutableDictionaryRef SourceAttributes = CFDictionaryCreateMutable(kCFAllocatorDefault, 1, &kCFTypeDictionaryKeyCallBacks, &kCFTypeDictionaryValueCallBacks);
+                CFDictionarySetValue(SourceAttributes, kCVPixelBufferOpenGLCompatibilityKey, kCFBooleanTrue);
+                CFDictionaryRef IOSurfaceValue = CFDictionaryCreate(kCFAllocatorDefault, nullptr, nullptr, 0, &kCFTypeDictionaryKeyCallBacks, &kCFTypeDictionaryValueCallBacks);
+                CFDictionarySetValue(SourceAttributes, kCVPixelBufferIOSurfacePropertiesKey, IOSurfaceValue);
+                // TODO (belchy06): This should support more than ARGB8 pixel format
+                int64 PixelType = kCVPixelFormatType_32BGRA;
+                CFNumberRef PixelFormat = CFNumberCreate(nullptr, kCFNumberLongType, &PixelType);
+                CFDictionarySetValue(SourceAttributes, kCVPixelBufferPixelFormatTypeKey, PixelFormat);
+
+                CONDITIONAL_RELEASE(IOSurfaceValue);
+                CONDITIONAL_RELEASE(PixelFormat);
+
+                VTDecompressionOutputCallbackRecord Record = 
+                {
+                    Internal::VTDecompressionOutputCallback, 
+                    this
+                };
+
+                OSStatus Result = VTDecompressionSessionCreate(kCFAllocatorDefault, PendingConfig.VideoFormat, nullptr, SourceAttributes, &Record, &Decoder);
+
+                CONDITIONAL_RELEASE(SourceAttributes);
+
+                if(Result != 0)
+                {
+                    DestroyDecompressionSession();
+                    return FAVResult(EAVResult::ErrorCreating, TEXT("Failed to create VTDecompressionSession"), TEXT("VT"), Result);
+                }
+
+                ConfigureDecompressionSession();
 			}
 		}
 
@@ -181,48 +148,42 @@ FAVResult TVideoDecoderVT<TResource>::SendPacket(FVideoPacket const& Packet)
 {
 	if (IsOpen())
 	{
-        // We've received our first call to decode a frame, we can now parse the information from
+        // We've received a call to decode a frame, we can now parse the information from
         // the bitstream, configure our config and initialize the session
-		if(!IsInitialized())
-        {
-            // Applying the config must be done before parsing the format as we need to know the codec
-            FAVResult AVResult = ApplyConfig();
-		    if (AVResult.IsNotSuccess())
-		    {
-    			return AVResult;
-		    }
-
-            CMVideoFormatDescriptionRef InputFormat = nullptr;
-            if(this->AppliedConfig.Codec == kCMVideoCodecType_H264)
-            {
-                InputFormat = NaluRewriter::CreateH264VideoFormatDescription(Packet.DataPtr.Get(), Packet.DataSize);
-            }
-            else if(this->AppliedConfig.Codec == kCMVideoCodecType_HEVC)
-            {
-                InputFormat = NaluRewriter::CreateH265VideoFormatDescription(Packet.DataPtr.Get(), Packet.DataSize);
-            }
-            else if (this->AppliedConfig.Codec == kCMVideoCodecType_VP9)
-            {
-                InputFormat = NaluRewriter::CreateVP9VideoFormatDescription(Packet.DataPtr.Get(), Packet.DataSize);
-            }
-            else
-            {
-                return FAVResult(EAVResult::Error, TEXT("Unsupported codec"), TEXT("VT"));
-            }
+        FVideoDecoderConfigVT& PendingConfig = this->EditPendingConfig();
 
-            if(InputFormat)
-            {
-                if(!CMFormatDescriptionEqual(InputFormat, VideoFormat))
-                {
-                    SetVideoFormat(InputFormat);
-                    ResetDecompressionSession();
-                }
-            }
+        CMVideoFormatDescriptionRef InputFormat = nullptr;
+        if(PendingConfig.Codec == kCMVideoCodecType_H264)
+        {
+            InputFormat = NaluRewriter::CreateH264VideoFormatDescription(Packet.DataPtr.Get(), Packet.DataSize);
+        }
+        else if(PendingConfig.Codec == kCMVideoCodecType_HEVC)
+        {
+            InputFormat = NaluRewriter::CreateH265VideoFormatDescription(Packet.DataPtr.Get(), Packet.DataSize);
+        }
+        else if (PendingConfig.Codec == kCMVideoCodecType_VP9)
+        {
+            InputFormat = NaluRewriter::CreateVP9VideoFormatDescription(Packet.DataPtr.Get(), Packet.DataSize);
+        }
+        else
+        {
+            return FAVResult(EAVResult::Error, TEXT("Unsupported codec"), TEXT("VT"));
+        }
 
-            CONDITIONAL_RELEASE(InputFormat);
+        if (InputFormat && !CMFormatDescriptionEqual(InputFormat, this->AppliedConfig.VideoFormat))
+        {
+            PendingConfig.SetVideoFormat(InputFormat);
+        }
+        
+        FAVResult AVResult = ApplyConfig();
+        if (AVResult.IsNotSuccess())
+        {
+            return AVResult;
         }
 
-        if(!VideoFormat)
+        CONDITIONAL_RELEASE(InputFormat);
+
+        if(!this->AppliedConfig.VideoFormat)
         {
             return FAVResult(EAVResult::WarningInvalidState, TEXT("Missing video format. Frame with sps/pps required."), TEXT("VT"));
         }
@@ -230,21 +191,21 @@ FAVResult TVideoDecoderVT<TResource>::SendPacket(FVideoPacket const& Packet)
         CMSampleBufferRef SampleBuffer = nullptr;
         if(this->AppliedConfig.Codec == kCMVideoCodecType_H264)
         {
-            if(!NaluRewriter::H264AnnexBBufferToCMSampleBuffer(Packet.DataPtr.Get(), Packet.DataSize, VideoFormat, &SampleBuffer, MemoryPool))
+            if(!NaluRewriter::H264AnnexBBufferToCMSampleBuffer(Packet.DataPtr.Get(), Packet.DataSize, this->AppliedConfig.VideoFormat, &SampleBuffer, MemoryPool))
             {
                 return FAVResult(EAVResult::Error, TEXT("Failed to get SampleBuffer"), TEXT("VT"));
             }
         }
         else if(this->AppliedConfig.Codec == kCMVideoCodecType_HEVC)
         {
-            if(!NaluRewriter::H265AnnexBBufferToCMSampleBuffer(Packet.DataPtr.Get(), Packet.DataSize, VideoFormat, &SampleBuffer, MemoryPool))
+            if(!NaluRewriter::H265AnnexBBufferToCMSampleBuffer(Packet.DataPtr.Get(), Packet.DataSize, this->AppliedConfig.VideoFormat, &SampleBuffer, MemoryPool))
             {
                 return FAVResult(EAVResult::Error, TEXT("Failed to get SampleBuffer"), TEXT("VT"));
             }
         }
         else if (this->AppliedConfig.Codec == kCMVideoCodecType_VP9)
         {
-            if(!NaluRewriter::VP9BufferToCMSampleBuffer(Packet.DataPtr.Get(), Packet.DataSize, VideoFormat, &SampleBuffer, MemoryPool))
+            if(!NaluRewriter::VP9BufferToCMSampleBuffer(Packet.DataPtr.Get(), Packet.DataSize, this->AppliedConfig.VideoFormat, &SampleBuffer, MemoryPool))
             {
                 return FAVResult(EAVResult::Error, TEXT("Failed to get SampleBuffer"), TEXT("VT"));
             }
@@ -282,28 +243,39 @@ FAVResult TVideoDecoderVT<TResource>::ReceiveFrame(TResolvableVideoResource<TRes
 {
     if(IsOpen())
     {
-        if(TSharedPtr<FFrame> Frame = *Frames.Peek())
+        if(Frames.Peek())
         {
+            TSharedPtr<FFrame> Frame = *Frames.Peek();
             size_t Width = CVPixelBufferGetWidth(Frame->ImageBuffer);
             size_t Height = CVPixelBufferGetHeight(Frame->ImageBuffer);
-
             OSType PixelFormat = CVPixelBufferGetPixelFormatType(Frame->ImageBuffer);
+            
+            static CVMetalTextureCacheRef TextureCache = nullptr;
+            if (TextureCache == nullptr)
+            {
+                id<MTLDevice> Device = (__bridge id<MTLDevice>)GDynamicRHI->RHIGetNativeDevice();
+                check(Device);
+                
+                CVReturn Result = CVMetalTextureCacheCreate(kCFAllocatorDefault, nullptr, Device, nullptr, &TextureCache);
+                if (Result != kCVReturnSuccess)
+                {
+                    return FAVResult(EAVResult::Error, TEXT("Failed to create CVMetalTextureCacheRef"), TEXT("VT"), Result);
+                }
+            }
+            
+            CVMetalTextureRef TextureRef = nullptr;
+            CVReturn Result = CVMetalTextureCacheCreateTextureFromImage(kCFAllocatorDefault, TextureCache, Frame->ImageBuffer, nullptr, MTLPixelFormatBGRA8Unorm_sRGB, Width, Height, 0, &TextureRef);
+            if (Result != kCVReturnSuccess)
+            {
+                return FAVResult(EAVResult::Error, TEXT("Failed to create CVMetalTextureRef"), TEXT("VT"), Result);
+            }
 
-            FVideoDescriptor ResourceDescriptor = FVideoDescriptor(EVideoFormat::BGRA, Width, Height);
+            FVideoDescriptor ResourceDescriptor = FVideoDescriptor(EVideoFormat::BGRA, Width, Height, new FBulkDataMetal(TextureRef));
             if (!InOutResource.Resolve(this->GetDevice(), ResourceDescriptor))
 			{
 				return FAVResult(EAVResult::ErrorResolving, TEXT("Failed to resolve frame resource"), TEXT("VT"));
 			}
 
-            CVPixelBufferLockBaseAddress(Frame->ImageBuffer, kCVPixelBufferLock_ReadOnly);
-
-            void* PixelPtr = CVPixelBufferGetBaseAddress(Frame->ImageBuffer);
-
-            // Do copy into the VideoResource
-//            InOutResource->GetRaw()->replaceRegion(MTL::Region(0, 0, Width, Height), 0, PixelPtr, Width * 4);
-
-            CVPixelBufferUnlockBaseAddress(Frame->ImageBuffer, kCVPixelBufferLock_ReadOnly);
-
             Frames.Pop();
             
             return EAVResult::Success;
@@ -315,23 +287,6 @@ FAVResult TVideoDecoderVT<TResource>::ReceiveFrame(TResolvableVideoResource<TRes
 	return FAVResult(EAVResult::ErrorInvalidState, TEXT("Decoder not open"), TEXT("VT"));
 }
 
-template <typename TResource>
-void TVideoDecoderVT<TResource>::SetVideoFormat(CMVideoFormatDescriptionRef Format)
-{
-    if(VideoFormat == Format)
-    {
-        return;
-    }
-
-    CONDITIONAL_RELEASE(VideoFormat);
-
-    VideoFormat = Format;
-    if(VideoFormat)
-    {
-        CFRetain(VideoFormat);
-    }
-}
-
 template <typename TResource>
 FAVResult TVideoDecoderVT<TResource>::HandleFrame(void* Params, OSStatus Status, VTDecodeInfoFlags InfoFlags, CVImageBufferRef ImageBuffer, CMTime Timestamp, CMTime Duration)
 {
diff --git a/Engine/Plugins/Media/PixelStreaming/Source/PixelStreaming/Private/PixelStreamingVideoSink.cpp b/Engine/Plugins/Media/PixelStreaming/Source/PixelStreaming/Private/PixelStreamingVideoSink.cpp
--- a/Engine/Plugins/Media/PixelStreaming/Source/PixelStreaming/Private/PixelStreamingVideoSink.cpp
+++ b/Engine/Plugins/Media/PixelStreaming/Source/PixelStreaming/Private/PixelStreamingVideoSink.cpp
@@ -25,7 +25,11 @@ void FPixelStreamingVideoSink::OnFrame(const webrtc::VideoFrame& Frame)
             return;
         }
 
-        TSharedPtr<FVideoResourceRHI, ESPMode::ThreadSafe> VideoResource = FrameBuffer->GetVideoResource()->TransformResource(FVideoDescriptor(EVideoFormat::BGRA, SizeX, SizeY));
+        TSharedPtr<FVideoResourceRHI, ESPMode::ThreadSafe> VideoResource = FrameBuffer->GetVideoResource();
+        if (VideoResource->GetFormat() != EVideoFormat::BGRA)
+        {
+            VideoResource = VideoResource->TransformResource(FVideoDescriptor(EVideoFormat::BGRA, SizeX, SizeY));
+        }
 		
         auto& Raw = StaticCastSharedPtr<FVideoResourceRHI>(VideoResource)->GetRaw();
         OnFrame(Raw.Texture);
@@ -92,4 +96,4 @@ void FPixelStreamingVideoSink::OnFrame(const webrtc::VideoFrame& Frame)
             OnFrame(SourceTexture);
     	});
 	}
-}
\ No newline at end of file
+}
