Fix: Build health issues from CL30257723

diff --git a/Engine/Plugins/Experimental/AVCodecs/VTCodecs/Source/VTCodecs/Public/Video/Decoders/VideoDecoderVT.hpp b/Engine/Plugins/Experimental/AVCodecs/VTCodecs/Source/VTCodecs/Public/Video/Decoders/VideoDecoderVT.hpp
--- a/Engine/Plugins/Experimental/AVCodecs/VTCodecs/Source/VTCodecs/Public/Video/Decoders/VideoDecoderVT.hpp
+++ b/Engine/Plugins/Experimental/AVCodecs/VTCodecs/Source/VTCodecs/Public/Video/Decoders/VideoDecoderVT.hpp
@@ -5,6 +5,8 @@
 #include "Video/Resources/Metal/VideoResourceMetal.h"
 #include "Video/Util/NaluRewriter.h"
 
+#include "DynamicRHI.h"
+
 THIRD_PARTY_INCLUDES_START
 #include <VideoToolbox/VideoToolbox.h>
 #include <CoreMedia/CMSync.h>
diff --git a/Engine/Plugins/Experimental/AVCodecs/VTCodecs/Source/VTCodecs/VTCodecs.build.cs b/Engine/Plugins/Experimental/AVCodecs/VTCodecs/Source/VTCodecs/VTCodecs.build.cs
--- a/Engine/Plugins/Experimental/AVCodecs/VTCodecs/Source/VTCodecs/VTCodecs.build.cs
+++ b/Engine/Plugins/Experimental/AVCodecs/VTCodecs/Source/VTCodecs/VTCodecs.build.cs
@@ -34,6 +34,7 @@ public class VTCodecs : ModuleRules
         PublicDependencyModuleNames.AddRange(new string[] {
             "RenderCore",
             "Core",
+            "RHI",
         });
 
         PublicFrameworks.AddRange(new string[]{
