Fix: VTCodecs outputting a key frame every frame regardless of encoder config

diff --git a/Engine/Plugins/Experimental/AVCodecs/AVCodecsCore/Source/AVCodecsCore/Private/Video/Resources/Metal/VideoResourceMetal.cpp b/Engine/Plugins/Experimental/AVCodecs/AVCodecsCore/Source/AVCodecsCore/Private/Video/Resources/Metal/VideoResourceMetal.cpp
--- a/Engine/Plugins/Experimental/AVCodecs/AVCodecsCore/Source/AVCodecsCore/Private/Video/Resources/Metal/VideoResourceMetal.cpp
+++ b/Engine/Plugins/Experimental/AVCodecs/AVCodecsCore/Source/AVCodecsCore/Private/Video/Resources/Metal/VideoResourceMetal.cpp
@@ -40,7 +40,7 @@ FVideoResourceMetal::FVideoResourceMetal(TSharedRef<FAVDevice> const& Device, CV
 {
     if (Raw)
     {
-        CFRetain(Raw);
+        CVPixelBufferRetain(Raw);
     }
 }
 
@@ -48,7 +48,7 @@ FVideoResourceMetal::~FVideoResourceMetal()
 {
     if (Raw)
     {
-        CFRelease(Raw);
+        CVPixelBufferRelease(Raw);
         Raw = nullptr;
     }
 }
diff --git a/Engine/Plugins/Experimental/AVCodecs/AVCodecsCore/Source/AVCodecsCoreRHI/Private/Video/Resources/VideoResourceRHI.cpp b/Engine/Plugins/Experimental/AVCodecs/AVCodecsCore/Source/AVCodecsCoreRHI/Private/Video/Resources/VideoResourceRHI.cpp
--- a/Engine/Plugins/Experimental/AVCodecs/AVCodecsCore/Source/AVCodecsCoreRHI/Private/Video/Resources/VideoResourceRHI.cpp
+++ b/Engine/Plugins/Experimental/AVCodecs/AVCodecsCore/Source/AVCodecsCoreRHI/Private/Video/Resources/VideoResourceRHI.cpp
@@ -565,84 +565,78 @@ DLLEXPORT FAVResult FAVExtension::TransformResource(TSharedPtr<FVideoResourceMet
 	if (InResource.IsValid())
 	{
 		if (InResource->GetDevice()->HasContext<FVideoContextMetal>())
-		{		
-            static CVMetalTextureCacheRef TextureCache = nullptr;
-            
-            if (TextureCache == nullptr)
+		{
+			FTextureRHIRef InResourceTexture = InResource->GetRaw().Texture;
+            if(bool(InResourceTexture->GetDesc().Flags & TexCreate_CPUReadback))
             {
-                id<MTLDevice> Device = (__bridge id<MTLDevice>)GDynamicRHI->RHIGetNativeDevice();
-                check(Device);
-                
-                CVReturn Result = CVMetalTextureCacheCreate(kCFAllocatorDefault, nullptr, Device, nullptr, &TextureCache);
+                /*
+                 * NOTE (william.belcher): If the texture passed in has was created with 'ETextureCreateFlags::CPUReadback', we can simply copy the raw bytes
+                 * out of the texture. This is the preferred way as it prevents an ~50ms block on the encoding thread
+                 */
+				const FVideoDescriptor& Descriptor = InResource->GetDescriptor();
+            
+            	CFMutableDictionaryRef SourceAttributes = CFDictionaryCreateMutable(kCFAllocatorDefault, 1, &kCFTypeDictionaryKeyCallBacks, &kCFTypeDictionaryValueCallBacks);
+            	CFDictionarySetValue(SourceAttributes, kCVPixelBufferOpenGLCompatibilityKey, kCFBooleanTrue);
+	
+            	CFDictionaryRef IOSurfaceValue = CFDictionaryCreate(kCFAllocatorDefault, nullptr, nullptr, 0, &kCFTypeDictionaryKeyCallBacks, &kCFTypeDictionaryValueCallBacks);
+            	CFDictionarySetValue(SourceAttributes, kCVPixelBufferIOSurfacePropertiesKey, IOSurfaceValue);
+	
+            	int64 PixelType = 0;
+            	switch(Descriptor.Format)
+            	{
+            	    case EVideoFormat::BGRA:
+            	        PixelType = kCVPixelFormatType_32BGRA;
+            	        break;
+            	    case EVideoFormat::ABGR10:
+            	        PixelType = kCVPixelFormatType_ARGB2101010LEPacked;
+            	        break;
+            	    default:
+            	        checkNoEntry();
+            	}
+	
+            	CFNumberRef PixelFormat = CFNumberCreate(nullptr, kCFNumberLongType, &PixelType);
+            	CFDictionarySetValue(SourceAttributes, kCVPixelBufferPixelFormatTypeKey, PixelFormat);
+	
+            	CFSafeRelease(IOSurfaceValue);
+            	CFSafeRelease(PixelFormat);
+	
+            	CVPixelBufferRef PixelBuffer;
+            	CVReturn Result = CVPixelBufferCreate(kCFAllocatorDefault, Descriptor.Width, Descriptor.Height, PixelType, SourceAttributes, &PixelBuffer);
+            	if (Result != kCVReturnSuccess)
+            	{
+            	    return FAVResult(EAVResult::Error, TEXT("Failed to create CVPixelBufferRef"), TEXT("RHI"), Result);
+            	}
+            	CFSafeRelease(SourceAttributes);
+
+                Result = CVPixelBufferLockBaseAddress(PixelBuffer, 0);
                 if (Result != kCVReturnSuccess)
                 {
-                    return FAVResult(EAVResult::Error, TEXT("Failed to create CVMetalTextureCacheRef"), TEXT("RHI"), Result);
+                    return FAVResult(EAVResult::Error, TEXT("Failed to lock base address"), TEXT("RHI"), Result);
                 }
-            }
-            
-            const FVideoDescriptor& Descriptor = InResource->GetDescriptor();
-                        
-            CFMutableDictionaryRef SourceAttributes = CFDictionaryCreateMutable(kCFAllocatorDefault, 1, &kCFTypeDictionaryKeyCallBacks, &kCFTypeDictionaryValueCallBacks);
-            CFDictionarySetValue(SourceAttributes, kCVPixelBufferOpenGLCompatibilityKey, kCFBooleanTrue);
-            CFDictionaryRef IOSurfaceValue = CFDictionaryCreate(kCFAllocatorDefault, nullptr, nullptr, 0, &kCFTypeDictionaryKeyCallBacks, &kCFTypeDictionaryValueCallBacks);
-            CFDictionarySetValue(SourceAttributes, kCVPixelBufferIOSurfacePropertiesKey, IOSurfaceValue);
-            int64 PixelType = kCVPixelFormatType_32BGRA;
-            CFNumberRef PixelFormat = CFNumberCreate(nullptr, kCFNumberLongType, &PixelType);
-            CFDictionarySetValue(SourceAttributes, kCVPixelBufferPixelFormatTypeKey, PixelFormat);
-
-            CFSafeRelease(IOSurfaceValue);
-            CFSafeRelease(PixelFormat);
-
-            CVPixelBufferRef PixelBuffer;
-            CVReturn Result = CVPixelBufferCreate(kCFAllocatorDefault, Descriptor.Width, Descriptor.Height, kCVPixelFormatType_32BGRA, SourceAttributes, &PixelBuffer);
-            if (Result != kCVReturnSuccess)
-            {
-                return FAVResult(EAVResult::Error, TEXT("Failed to create CVPixelBufferRef"), TEXT("RHI"), Result);
-            }
-            CFSafeRelease(SourceAttributes);
 
-            CVMetalTextureRef TextureRef;
-            Result = CVMetalTextureCacheCreateTextureFromImage(kCFAllocatorDefault, TextureCache, PixelBuffer, nullptr, MTLPixelFormatBGRA8Unorm_sRGB, Descriptor.Width, Descriptor.Height, 0, &TextureRef);
-            if (Result != kCVReturnSuccess)
-            {
-                return FAVResult(EAVResult::Error, TEXT("Failed to create CVMetalTextureRef"), TEXT("RHI"), Result);
-            }
+                // NOTE (belchy06): GetBytes assumes the raw texture has been created with with TexCreate_CPUReadback
+                static_cast<MTL::Texture*>(InResource->GetRaw().Texture->GetNativeResource())->getBytes(reinterpret_cast<uint8*>(CVPixelBufferGetBaseAddressOfPlane(PixelBuffer, 0)), CVPixelBufferGetBytesPerRow(PixelBuffer), MTL::Region(0, 0, Descriptor.Width, Descriptor.Height), 0);
+                CVPixelBufferUnlockBaseAddress(PixelBuffer, 0);
+                
+                OutResource = MakeShared<FVideoResourceMetal>(InResource->GetDevice(),
+                                                              PixelBuffer,
+                                                              InResource->GetLayout());
 
-            FGPUFenceRHIRef Fence;
-            FTextureRHIRef Destination;
-            const FRHITextureCreateDesc Desc =
-                FRHITextureCreateDesc::Create2D(TEXT("FAVExtension::TransformResource"), Descriptor.Width, Descriptor.Height, PF_B8G8R8A8)
-                .SetFlags(ETextureCreateFlags::SRGB | ETextureCreateFlags::Dynamic | ETextureCreateFlags::NoTiling | ETextureCreateFlags::ShaderResource)
-                .SetBulkData(new FBulkDataMetal(TextureRef));
-                        
-            ENQUEUE_RENDER_COMMAND(FAVExtensionTransformResource)(
-                [&Source = InResource->GetRaw().Texture, &Desc, &Fence, &Destination](FRHICommandListImmediate& RHICmdList)
-            {
-                Fence = RHICreateGPUFence(TEXT("FAVExtension::TransformResource"));
-                Destination = RHICreateTexture(Desc);
-                RHICmdList.CopyTexture(Source, Destination, FRHICopyTextureInfo());
-                            
-                RHICmdList.WriteGPUFence(Fence);
-            });
-            
-            while(!Fence || !Fence->Poll())
+				CVPixelBufferRelease(PixelBuffer);
+			
+			    return OutResource->Validate();
+            }
+            else
             {
-				if(IsEngineExitRequested())
+                static bool bLogWarning = true;
+				if(bLogWarning)
 				{
-					return FAVResult(EAVResult::Error, TEXT("Engine exit requested before texture copy was complete"), TEXT("RHI"));
+					bLogWarning = false;
+					FAVResult::Log(EAVResult::Warning, TEXT("Unable to transform video resource! Metal RHI requires FVideoResourceRHI textures be created with the ETextureCreateFlags::CPUReadback flag"), TEXT("RHI"));
 				}
 
-                FPlatformProcess::YieldThread();
+				return EAVResult::Error;
             }
-
-            OutResource = MakeShared<FVideoResourceMetal>(InResource->GetDevice(),
-                                                          PixelBuffer,
-                                                          InResource->GetLayout());
-
-            CFSafeRelease(TextureRef);
-            CFSafeRelease(PixelBuffer);
-            
-            return OutResource->Validate();
 		}
 
 		return FAVResult(EAVResult::ErrorMapping, TEXT("No Metal context found"), TEXT("RHI"));
diff --git a/Engine/Plugins/Experimental/AVCodecs/VTCodecs/Source/VTCodecs/Private/Video/Util/VTSessionHelpers.cpp b/Engine/Plugins/Experimental/AVCodecs/VTCodecs/Source/VTCodecs/Private/Video/Util/VTSessionHelpers.cpp
--- a/Engine/Plugins/Experimental/AVCodecs/VTCodecs/Source/VTCodecs/Private/Video/Util/VTSessionHelpers.cpp
+++ b/Engine/Plugins/Experimental/AVCodecs/VTCodecs/Source/VTCodecs/Private/Video/Util/VTSessionHelpers.cpp
@@ -30,19 +30,6 @@ void VTSessionHelpers::SetVTSessionProperty(VTSessionRef Session, CFStringRef Ke
     }
 }
 
-void VTSessionHelpers::SetVTSessionProperty(VTSessionRef Session, CFStringRef Key, uint32_t Value) 
-{
-    int64_t Value_64 = Value;
-    CFNumberRef Num = CFNumberCreate(kCFAllocatorDefault, kCFNumberSInt64Type, &Value_64);
-    OSStatus Status = VTSessionSetProperty(Session, Key, Num);
-    CFRelease(Num);
-    if (Status != 0) 
-    {
-        FString KeyString = CFStringToString(Key);
-        FAVResult::Log(EAVResult::Warning, FString::Printf(TEXT("VTSessionSetProperty failed to set: %s to %d"), *KeyString, Value), TEXT("VT"), Status);
-    }
-}
-
 void VTSessionHelpers::SetVTSessionProperty(VTSessionRef Session, CFStringRef Key, bool Value) 
 {
     CFBooleanRef Bool = (Value) ? kCFBooleanTrue : kCFBooleanFalse;
diff --git a/Engine/Plugins/Experimental/AVCodecs/VTCodecs/Source/VTCodecs/Public/Video/Encoders/VideoEncoderVT.hpp b/Engine/Plugins/Experimental/AVCodecs/VTCodecs/Source/VTCodecs/Public/Video/Encoders/VideoEncoderVT.hpp
--- a/Engine/Plugins/Experimental/AVCodecs/VTCodecs/Source/VTCodecs/Public/Video/Encoders/VideoEncoderVT.hpp
+++ b/Engine/Plugins/Experimental/AVCodecs/VTCodecs/Source/VTCodecs/Public/Video/Encoders/VideoEncoderVT.hpp
@@ -111,36 +111,45 @@ FAVResult TVideoEncoderVT<TResource>::ApplyConfig()
 
                 CFDictionaryRef IOSurfaceValue = CFDictionaryCreate(kCFAllocatorDefault, nullptr, nullptr, 0, &kCFTypeDictionaryKeyCallBacks, &kCFTypeDictionaryValueCallBacks);
                
-                int64 PixelType = PendingConfig.GetCVPixelFormatType();
+                int64_t PixelType = PendingConfig.GetCVPixelFormatType();
                 CFNumberRef PixelFormat = CFNumberCreate(nullptr, kCFNumberLongType, &PixelType);
 
-                CFTypeRef Values[AttributesSize] = {kCFBooleanTrue, IOSurfaceValue, PixelFormat};
+                CFTypeRef Values[AttributesSize] = { kCFBooleanTrue, IOSurfaceValue, PixelFormat };
                 CFDictionaryRef SourceAttributes = CFDictionaryCreate(kCFAllocatorDefault, Keys, Values, AttributesSize, &kCFTypeDictionaryKeyCallBacks, &kCFTypeDictionaryValueCallBacks);
 
                 CONDITIONAL_RELEASE(IOSurfaceValue);
                 CONDITIONAL_RELEASE(PixelFormat);
+          
+                OSStatus Result = VTCompressionSessionCreate(kCFAllocatorDefault, 
+                                                             PendingConfig.Width,
+                                                             PendingConfig.Height,
+                                                             PendingConfig.Codec,
+                                                             NULL,
+                                                             SourceAttributes,
+                                                             NULL /* Default Compressed Data Allocator */,
+                                                             Internal::VTCompressionOutputCallback,
+                                                             this,
+                                                             &Encoder);
 
-                CFMutableDictionaryRef EncoderSpecification = CFDictionaryCreateMutable(kCFAllocatorDefault, 1, &kCFTypeDictionaryKeyCallBacks, &kCFTypeDictionaryValueCallBacks);
-                
-                // We explicitly state here that we want a hardware encoder. If it's optional, replace with kVTVideoEncoderSpecification_EnableHardwareAcceleratedVideoEncoder
-                CFDictionarySetValue(EncoderSpecification, kVTVideoEncoderSpecification_RequireHardwareAcceleratedVideoEncoder, kCFBooleanTrue);
-                
-                OSStatus Result = VTCompressionSessionCreate(kCFAllocatorDefault, PendingConfig.Width, PendingConfig.Height, PendingConfig.Codec, EncoderSpecification, SourceAttributes, nullptr /* Default Compressed Data Allocator */, Internal::VTCompressionOutputCallback, this, &Encoder);
+                CONDITIONAL_RELEASE(SourceAttributes);
                     
                 if(Result != 0)
                 {
-                    CONDITIONAL_RELEASE(SourceAttributes);
-                    CONDITIONAL_RELEASE(EncoderSpecification);
                     return FAVResult(EAVResult::ErrorCreating, TEXT("Failed to create VTCompressionSession"), TEXT("VT"), Result);
                 }
 
                 CFBooleanRef bIsUsingHardwareEncoder;
                 Result = VTSessionCopyProperty(Encoder, kVTCompressionPropertyKey_UsingHardwareAcceleratedVideoEncoder, NULL, &bIsUsingHardwareEncoder);
-                FAVResult::Log(EAVResult::Unknown, FString::Printf(TEXT("Created compression session. UsingHardwareEncoder: %s"), (Result == 0 && CFBooleanGetValue(bIsUsingHardwareEncoder) ? TEXT("TRUE") : TEXT("FALSE"))), TEXT("VT"));
-
+                if(Result != 0)
+                {
+                    FAVResult::Log(EAVResult::Warning, TEXT("VTSessionCopyProperty(kVTCompressionPropertyKey_UsingHardwareAcceleratedVideoEncoder) failed"), TEXT("VT"), Result);
+                }
+                else
+                {
+                    FAVResult::Log(EAVResult::Warning, FString::Printf(TEXT("Created compression session. UsingHardwareEncoder: %s"), (CFBooleanGetValue(bIsUsingHardwareEncoder) ? TEXT("TRUE") : TEXT("FALSE"))), TEXT("VT"));
+                }
+                
                 CONDITIONAL_RELEASE(bIsUsingHardwareEncoder);
-                CONDITIONAL_RELEASE(SourceAttributes);
-                CONDITIONAL_RELEASE(EncoderSpecification);
 
                 ConfigureCompressionSession(PendingConfig);
 			}
@@ -154,33 +163,27 @@ FAVResult TVideoEncoderVT<TResource>::ApplyConfig()
 
 template <typename TResource>
 FAVResult TVideoEncoderVT<TResource>::ConfigureCompressionSession(FVideoEncoderConfigVT const& Config)  
-{   
+{       
     VTSessionHelpers::SetVTSessionProperty(Encoder, kVTCompressionPropertyKey_RealTime, true);
-    
+
     VTSessionHelpers::SetVTSessionProperty(Encoder, kVTCompressionPropertyKey_ProfileLevel, Config.Profile);
     
-    VTSessionHelpers::SetVTSessionProperty(Encoder, kVTCompressionPropertyKey_AllowFrameReordering, false);
-    
     SetEncoderBitrate(Config);
     
-    VTSessionHelpers::SetVTSessionProperty(Encoder, kVTCompressionPropertyKey_MaxKeyFrameInterval, Config.KeyframeInterval);
+    VTSessionHelpers::SetVTSessionProperty(Encoder, kVTCompressionPropertyKey_AllowTemporalCompression, false);
     
-    VTSessionHelpers::SetVTSessionProperty(Encoder, kVTCompressionPropertyKey_MaxKeyFrameIntervalDuration, Config.KeyframeInterval * Config.FrameRate);
+    VTSessionHelpers::SetVTSessionProperty(Encoder, kVTCompressionPropertyKey_AllowFrameReordering, false);
     
-    VTSessionHelpers::SetVTSessionProperty(Encoder, kVTCompressionPropertyKey_MinAllowedFrameQP, Config.MinQP);
+    VTSessionHelpers::SetVTSessionProperty(Encoder, kVTCompressionPropertyKey_MaxKeyFrameInterval, (int32_t)Config.KeyframeInterval);
     
-    VTSessionHelpers::SetVTSessionProperty(Encoder, kVTCompressionPropertyKey_MaxAllowedFrameQP, Config.MaxQP);
+    VTSessionHelpers::SetVTSessionProperty(Encoder, kVTCompressionPropertyKey_MaxKeyFrameIntervalDuration, (Config.KeyframeInterval > 0 ? (int32_t)(Config.KeyframeInterval / Config.FrameRate) : 0));
     
+    VTSessionHelpers::SetVTSessionProperty(Encoder, kVTCompressionPropertyKey_MaximizePowerEfficiency, false);
+
     if(Config.Codec == kCMVideoCodecType_H264)
     {
         VTSessionHelpers::SetVTSessionProperty(Encoder, kVTCompressionPropertyKey_H264EntropyMode, Config.EntropyCodingMode);
     }
-    
-    OSStatus Status = VTCompressionSessionPrepareToEncodeFrames(Encoder);
-    if(Status != 0)
-    {
-        return FAVResult(EAVResult::Warning, TEXT("Compression session failed to prepare encode frames"), TEXT("VT"));
-    }
 
     return EAVResult::Success;
 }
@@ -189,29 +192,7 @@ template <typename TResource>
 FAVResult TVideoEncoderVT<TResource>::SetEncoderBitrate(FVideoEncoderConfigVT const& Config)
 {
     VTSessionHelpers::SetVTSessionProperty(Encoder, kVTCompressionPropertyKey_AverageBitRate, Config.TargetBitrate);
-
-    // Bits to bytes conversion
-    int64_t BytesPerSecondValue = static_cast<int64_t>(Config.MaxBitrate >> 3);
-    CFNumberRef BytesPerSecond = CFNumberCreate(kCFAllocatorDefault, kCFNumberSInt64Type, &BytesPerSecondValue);
-    
-    int64_t OneSecondValue = 1;
-    CFNumberRef OneSecond = CFNumberCreate(kCFAllocatorDefault, kCFNumberSInt64Type, &OneSecondValue);
-    
-    const void* Nums[2] = { BytesPerSecond, OneSecond };
-    CFArrayRef DataRateLimits = CFArrayCreate(nullptr, Nums, 2, &kCFTypeArrayCallBacks);
-    
-    OSStatus Result = VTSessionSetProperty(Encoder, kVTCompressionPropertyKey_DataRateLimits, DataRateLimits);
-    
-    CONDITIONAL_RELEASE(BytesPerSecond);
-    CONDITIONAL_RELEASE(OneSecond);
-    CONDITIONAL_RELEASE(DataRateLimits);
-    
-    if(Result != 0)
-    {
-        return FAVResult(EAVResult::Error, TEXT("Failed to set kVTCompressionPropertyKey_DataRateLimits"), TEXT("VT"), Result);
-    }
     
-
     return EAVResult::Success;
 }
 
@@ -224,16 +205,16 @@ FAVResult TVideoEncoderVT<TResource>::SendFrame(TSharedPtr<FVideoResourceMetal>
         PendingConfig.PixelFormat = Resource->GetFormat();
 
         FAVResult AVResult = ApplyConfig();
-		if (AVResult.IsNotSuccess())
-		{
-			return AVResult;
-		}
+        if (AVResult.IsNotSuccess())
+        {
+            return AVResult;
+        }
 
         if(Resource.IsValid())
         {
-            CMTime PresentationTime = CMTimeMake(Timestamp, 1000);
+            CMTime PresentationTime = CMTimeMake(Timestamp, 1000000);
             CFDictionaryRef FrameProperties = nullptr;
-            if (bForceKeyframe) 
+            if (bForceKeyframe)
             {
                 CFTypeRef Keys[] = {kVTEncodeFrameOptionKey_ForceKeyFrame};
                 CFTypeRef Values[] = {kCFBooleanTrue};
@@ -243,13 +224,13 @@ FAVResult TVideoEncoderVT<TResource>::SendFrame(TSharedPtr<FVideoResourceMetal>
             TUniquePtr<EncodeParams> Params = MakeUnique<EncodeParams>();
             Params.Reset(new EncodeParams(PendingConfig.Codec, PresentationTime));
 
-            OSStatus Status = VTCompressionSessionEncodeFrame(Encoder, Resource->GetRaw(), PresentationTime, kCMTimeInvalid, FrameProperties, (void*)Params.Release(), nullptr);
+            OSStatus Status = VTCompressionSessionEncodeFrame(Encoder, Resource->GetRaw(), PresentationTime, kCMTimeInvalid, FrameProperties, (void*)Params.Release(), NULL);
 
-            CONDITIONAL_RELEASE(FrameProperties);        
+            CONDITIONAL_RELEASE(FrameProperties);
 
-            if(Status != 0)
+            if(Status != kCVReturnSuccess)
             {
-                return FAVResult(EAVResult::Error, TEXT("Failed to encode"), TEXT("VT"));
+                return FAVResult(EAVResult::Error, TEXT("Failed to encode"), TEXT("VT"), Status);
             }
             
             return EAVResult::Success;
@@ -258,19 +239,18 @@ FAVResult TVideoEncoderVT<TResource>::SendFrame(TSharedPtr<FVideoResourceMetal>
         {
             // Flush encoder
             OSStatus Status = VTCompressionSessionCompleteFrames(Encoder, kCMTimeInvalid);
-            if(Status != 0)
+            if(Status != kCVReturnSuccess)
             {
-                return FAVResult(EAVResult::Error, TEXT("Failed to flush"), TEXT("VT"));
+                return FAVResult(EAVResult::Error, TEXT("Failed to flush"), TEXT("VT"), Status);
             }
 
             return EAVResult::Success;
         }
     }
 
-	return FAVResult(EAVResult::ErrorInvalidState, TEXT("Encoder not open"), TEXT("VT"));
+    return FAVResult(EAVResult::ErrorInvalidState, TEXT("Encoder not open"), TEXT("VT"));
 }
 
-
 template <typename TResource>
 FAVResult TVideoEncoderVT<TResource>::ReceivePacket(FVideoPacket& OutPacket)
 {
@@ -294,12 +274,12 @@ FAVResult TVideoEncoderVT<TResource>::HandlePacket(void* Params, OSStatus Status
 	{
         if(Status != 0)
         {
-            return FAVResult(EAVResult::Error, TEXT("Failed to encode"), TEXT("VT"));
+            return FAVResult(EAVResult::Error, TEXT("Failed to encode"), TEXT("VT"), Status);
         }
 
         if(InfoFlags & kVTEncodeInfo_FrameDropped)
         {
-            return FAVResult(EAVResult::Error, TEXT("Frame dropped"), TEXT("VT"));
+            return FAVResult(EAVResult::Error, TEXT("Frame dropped"), TEXT("VT"), Status);
         }
 
         bool bIsKeyframe = false;
@@ -344,7 +324,7 @@ FAVResult TVideoEncoderVT<TResource>::HandlePacket(void* Params, OSStatus Status
             FAVResult Result = H264.Parse(FVideoPacket(CopiedData, Bitstream.Num(), 0, 0, 0, false), Slices);
             if(Result != EAVResult::Success)
             {
-                return FAVResult(EAVResult::Error, TEXT("Failed to parse bitstream"), TEXT("VT"));
+                return FAVResult(EAVResult::Error, TEXT("Failed to parse bitstream"), TEXT("VT"), Result);
             }
 
             QP = H264.GetLastSliceQP(Slices).Get(0);
@@ -356,7 +336,7 @@ FAVResult TVideoEncoderVT<TResource>::HandlePacket(void* Params, OSStatus Status
             FAVResult Result = H265.Parse(FVideoPacket(CopiedData, Bitstream.Num(), 0, 0, 0, false), Slices);
             if(Result != EAVResult::Success)
             {
-                return FAVResult(EAVResult::Error, TEXT("Failed to parse bitstream"), TEXT("VT"));
+                return FAVResult(EAVResult::Error, TEXT("Failed to parse bitstream"), TEXT("VT"), Result);
             }
 
             QP = H265.GetLastSliceQP(Slices).Get(0);
diff --git a/Engine/Plugins/Experimental/AVCodecs/VTCodecs/Source/VTCodecs/Public/Video/Util/VTSessionHelpers.h b/Engine/Plugins/Experimental/AVCodecs/VTCodecs/Source/VTCodecs/Public/Video/Util/VTSessionHelpers.h
--- a/Engine/Plugins/Experimental/AVCodecs/VTCodecs/Source/VTCodecs/Public/Video/Util/VTSessionHelpers.h
+++ b/Engine/Plugins/Experimental/AVCodecs/VTCodecs/Source/VTCodecs/Public/Video/Util/VTSessionHelpers.h
@@ -14,7 +14,6 @@ class VTSessionHelpers
 public:
     static FString CFStringToString(const CFStringRef CfString);
 	static void SetVTSessionProperty(VTSessionRef Session, CFStringRef Key, int32_t Value);
-	static void SetVTSessionProperty(VTSessionRef Session, CFStringRef Key, uint32_t Value);
 	static void SetVTSessionProperty(VTSessionRef Session, CFStringRef Key, bool Value);
 	static void SetVTSessionProperty(VTSessionRef Session, CFStringRef Key, CFStringRef Value);
 };
