Fix: Add runtime query for HW encoder support in VideoToolbox

diff --git a/Engine/Plugins/Experimental/AVCodecs/VTCodecs/Source/VTCodecs/Private/VTCodecsModule.cpp b/Engine/Plugins/Experimental/AVCodecs/VTCodecs/Source/VTCodecs/Private/VTCodecsModule.cpp
--- a/Engine/Plugins/Experimental/AVCodecs/VTCodecs/Source/VTCodecs/Private/VTCodecsModule.cpp
+++ b/Engine/Plugins/Experimental/AVCodecs/VTCodecs/Source/VTCodecs/Private/VTCodecsModule.cpp
@@ -23,16 +23,69 @@ class FVTModule : public IModuleInterface
 public:
 	virtual void StartupModule() override
 	{       
-        VTRegisterSupplementalVideoDecoderIfAvailable(kCMVideoCodecType_VP9);
+        TFunction<bool(TSharedRef<FAVDevice> const&, const TCHAR*)> CheckCodecSupport = [](TSharedRef<FAVDevice> const& Device, const TCHAR* Codec) {
+			CFStringRef TargetCodec = CFStringCreateWithCString(kCFAllocatorDefault, TCHAR_TO_ANSI(Codec), kCFStringEncodingUTF8);
+            CFArrayRef EncoderList = NULL;
+            bool bSupportsHardware = false;
+
+            OSStatus Status = VTCopyVideoEncoderList(NULL, &EncoderList);
+            if (Status != 0) 
+            {
+                return false;
+            }
+
+            CFIndex EncoderCount = CFArrayGetCount(EncoderList);
+            for (CFIndex i = 0; i < EncoderCount; i++) 
+            {
+				CFDictionaryRef EncoderDict = (CFDictionaryRef)CFArrayGetValueAtIndex(EncoderList, i);
+                if (EncoderDict)
+                {
+					CFStringRef EncoderCodec = (CFStringRef)CFDictionaryGetValue(EncoderDict, kVTVideoEncoderList_CodecName);
+		 			if (EncoderCodec != TargetCodec)
+					{
+						continue;
+					}
+					
+                    CFBooleanRef bIsHardware = (CFBooleanRef)CFDictionaryGetValue(EncoderDict, kVTVideoEncoderList_IsHardwareAccelerated);
+					bSupportsHardware |= bIsHardware ? (bool)CFBooleanGetValue(bIsHardware) : false;
+                }
+            }
+
+            if (EncoderList) 
+            {
+                CFRelease(EncoderList);
+            }
+			
+			if(TargetCodec)
+			{
+				CFRelease(TargetCodec);
+			}
+
+			return bSupportsHardware;
+		};
 
         FVideoEncoder::RegisterPermutationsOf<TVideoEncoderVT<FVideoResourceMetal>>
             ::With<FVideoResourceMetal>
-            ::And<FVideoEncoderConfigVT, FVideoEncoderConfigH264, FVideoEncoderConfigH265>(
-                [](TSharedRef<FAVDevice> const& NewDevice, TSharedRef<FAVInstance> const& NewInstance)
+            ::And<FVideoEncoderConfigVT, FVideoEncoderConfigH264>(
+                [CheckCodecSupport](TSharedRef<FAVDevice> const& NewDevice, TSharedRef<FAVInstance> const& NewInstance)
                 {
-                    return FAPI::Get<FVT>().IsValid();
+					static bool bSupportsCodec = CheckCodecSupport(NewDevice, TEXT("H.264"));
+
+                    return FAPI::Get<FVT>().IsValid() && bSupportsCodec;
                 });
 
+        FVideoEncoder::RegisterPermutationsOf<TVideoEncoderVT<FVideoResourceMetal>>
+            ::With<FVideoResourceMetal>
+            ::And<FVideoEncoderConfigVT, FVideoEncoderConfigH265>(
+                [CheckCodecSupport](TSharedRef<FAVDevice> const& NewDevice, TSharedRef<FAVInstance> const& NewInstance)
+                {
+					static bool bSupportsCodec = CheckCodecSupport(NewDevice, TEXT("HEVC"));
+
+                    return FAPI::Get<FVT>().IsValid() && bSupportsCodec;
+                });
+
+        VTRegisterSupplementalVideoDecoderIfAvailable(kCMVideoCodecType_VP9);
+
         if(VTIsHardwareDecodeSupported(kCMVideoCodecType_H264))
         {
             FVideoDecoder::RegisterPermutationsOf<TVideoDecoderVT<FVideoResourceMetal>>
diff --git a/Engine/Plugins/Experimental/AVCodecs/VTCodecs/Source/VTCodecs/Public/Video/Encoders/VideoEncoderVT.hpp b/Engine/Plugins/Experimental/AVCodecs/VTCodecs/Source/VTCodecs/Public/Video/Encoders/VideoEncoderVT.hpp
--- a/Engine/Plugins/Experimental/AVCodecs/VTCodecs/Source/VTCodecs/Public/Video/Encoders/VideoEncoderVT.hpp
+++ b/Engine/Plugins/Experimental/AVCodecs/VTCodecs/Source/VTCodecs/Public/Video/Encoders/VideoEncoderVT.hpp
@@ -151,6 +151,10 @@ FAVResult TVideoEncoderVT<TResource>::ApplyConfig()
                     
                 if(Result != 0)
                 {
+                    if(Result == kVTInvalidSessionErr)
+                    {
+                        FAVResult::Log(EAVResult::Warning, TEXT("Invalid VTCompressionSession. You could be exceeding the maximum resolution supported by VideoToolbox!"), TEXT("VT"));
+                    }
                     return FAVResult(EAVResult::ErrorCreating, TEXT("Failed to create VTCompressionSession"), TEXT("VT"), Result);
                 }
 
