Fix missed mac support when updating AVCodecs from QP to Quality

#rnx

diff --git a/Engine/Plugins/Experimental/AVCodecs/VTCodecs/Source/VTCodecs/Private/Video/Encoders/Configs/VideoEncoderConfigVT.cpp b/Engine/Plugins/Experimental/AVCodecs/VTCodecs/Source/VTCodecs/Private/Video/Encoders/Configs/VideoEncoderConfigVT.cpp
--- a/Engine/Plugins/Experimental/AVCodecs/VTCodecs/Source/VTCodecs/Private/Video/Encoders/Configs/VideoEncoderConfigVT.cpp
+++ b/Engine/Plugins/Experimental/AVCodecs/VTCodecs/Source/VTCodecs/Private/Video/Encoders/Configs/VideoEncoderConfigVT.cpp
@@ -2,7 +2,6 @@
 
 #include "Video/Encoders/Configs/VideoEncoderConfigVT.h"
 
-
 #include "Video/Encoders/Configs/VideoEncoderConfigH264.h"
 #include "Video/Encoders/Configs/VideoEncoderConfigH265.h"
 
@@ -11,23 +10,23 @@ REGISTER_TYPEID(FVideoEncoderConfigVT);
 static uint32 DEFAULT_BITRATE_TARGET = 1000000;
 static uint32 DEFAULT_BITRATE_MAX = 10000000;
 
-TAVResult<OSType> FVideoEncoderConfigVT::ConvertFormat(EVideoFormat const& Format)
+TAVResult<OSType> FVideoEncoderConfigVT::ConvertFormat(EVideoFormat const &Format)
 {
 	switch (Format)
 	{
-		case EVideoFormat::BGRA:
-			return kCVPixelFormatType_32RGBA;
-		case EVideoFormat::ABGR10:
-			return kCVPixelFormatType_ARGB2101010LEPacked;
-		case EVideoFormat::NV12:
-			return kCVPixelFormatType_420YpCbCr8BiPlanarFullRange;
-		default:
-			return FAVResult(EAVResult::ErrorUnsupported, FString::Printf(TEXT("Pixel format %d is not supported"), Format), TEXT("VT"));
+	case EVideoFormat::BGRA:
+		return kCVPixelFormatType_32RGBA;
+	case EVideoFormat::ABGR10:
+		return kCVPixelFormatType_ARGB2101010LEPacked;
+	case EVideoFormat::NV12:
+		return kCVPixelFormatType_420YpCbCr8BiPlanarFullRange;
+	default:
+		return FAVResult(EAVResult::ErrorUnsupported, FString::Printf(TEXT("Pixel format %d is not supported"), Format), TEXT("VT"));
 	}
 }
 
-template<>
-DLLEXPORT FAVResult FAVExtension::TransformConfig(FVideoEncoderConfigVT& OutConfig, FVideoEncoderConfig const& InConfig)
+template <>
+DLLEXPORT FAVResult FAVExtension::TransformConfig(FVideoEncoderConfigVT &OutConfig, FVideoEncoderConfig const &InConfig)
 {
 	OutConfig.Width = InConfig.Width;
 	OutConfig.Height = InConfig.Height;
@@ -36,15 +35,15 @@ DLLEXPORT FAVResult FAVExtension::TransformConfig(FVideoEncoderConfigVT& OutConf
 	OutConfig.TargetBitrate = InConfig.TargetBitrate > -1 ? InConfig.TargetBitrate : DEFAULT_BITRATE_TARGET;
 	OutConfig.MaxBitrate = InConfig.MaxBitrate > -1 ? InConfig.MaxBitrate : DEFAULT_BITRATE_MAX;
 	OutConfig.RateControlMode = InConfig.RateControlMode;
-    OutConfig.MinQP = InConfig.MinQP;
-    OutConfig.MaxQP = InConfig.MaxQP;
+	OutConfig.MinQuality = InConfig.MinQuality;
+	OutConfig.MaxQuality = InConfig.MaxQuality;
 	OutConfig.KeyframeInterval = InConfig.KeyframeInterval > 0 ? InConfig.KeyframeInterval : 0;
 
 	return EAVResult::Success;
 }
 
 template <>
-DLLEXPORT FAVResult FAVExtension::TransformConfig(FVideoEncoderConfig& OutConfig, FVideoEncoderConfigVT const& InConfig)
+DLLEXPORT FAVResult FAVExtension::TransformConfig(FVideoEncoderConfig &OutConfig, FVideoEncoderConfigVT const &InConfig)
 {
 	OutConfig.Width = InConfig.Width;
 	OutConfig.Height = InConfig.Height;
@@ -53,39 +52,40 @@ DLLEXPORT FAVResult FAVExtension::TransformConfig(FVideoEncoderConfig& OutConfig
 	OutConfig.TargetBitrate = InConfig.TargetBitrate > -1 ? InConfig.TargetBitrate : DEFAULT_BITRATE_TARGET;
 	OutConfig.MaxBitrate = InConfig.MaxBitrate > -1 ? InConfig.MaxBitrate : DEFAULT_BITRATE_MAX;
 	OutConfig.RateControlMode = InConfig.RateControlMode;
-    OutConfig.MinQP = InConfig.MinQP;
-    OutConfig.MaxQP = InConfig.MaxQP;
+	OutConfig.MinQuality = InConfig.MinQuality;
+	OutConfig.MaxQuality = InConfig.MaxQuality;
 	OutConfig.KeyframeInterval = InConfig.KeyframeInterval > 0 ? InConfig.KeyframeInterval : 0;
 
 	return EAVResult::Success;
 }
 
 template <>
-DLLEXPORT FAVResult FAVExtension::TransformConfig(FVideoEncoderConfigVT& OutConfig, FVideoEncoderConfigH264 const& InConfig)
+DLLEXPORT FAVResult FAVExtension::TransformConfig(FVideoEncoderConfigVT &OutConfig, FVideoEncoderConfigH264 const &InConfig)
 {
-    OutConfig.Codec = kCMVideoCodecType_H264;
+	OutConfig.Codec = kCMVideoCodecType_H264;
 
-    static auto const ConvertEntropyCodingMode = [](EH264EntropyCodingMode EntropyCodingMode) -> TAVResult<CFStringRef> {
+	static auto const ConvertEntropyCodingMode = [](EH264EntropyCodingMode EntropyCodingMode) -> TAVResult<CFStringRef>
+	{
 		switch (EntropyCodingMode)
 		{
-			case EH264EntropyCodingMode::Auto:
-				return nullptr;
-			case EH264EntropyCodingMode::CABAC:
-				return kVTH264EntropyMode_CABAC;
-			case EH264EntropyCodingMode::CAVLC:
-				return kVTH264EntropyMode_CAVLC;
-			default:
-				return FAVResult(EAVResult::ErrorUnsupported, FString::Printf(TEXT("H264 entropy coding mode %d is not supported"), EntropyCodingMode), TEXT("VT"));
+		case EH264EntropyCodingMode::Auto:
+			return nullptr;
+		case EH264EntropyCodingMode::CABAC:
+			return kVTH264EntropyMode_CABAC;
+		case EH264EntropyCodingMode::CAVLC:
+			return kVTH264EntropyMode_CAVLC;
+		default:
+			return FAVResult(EAVResult::ErrorUnsupported, FString::Printf(TEXT("H264 entropy coding mode %d is not supported"), EntropyCodingMode), TEXT("VT"));
 		}
 	};
 
-    TAVResult<CFStringRef> const ConvertedEntropyCodingMode = ConvertEntropyCodingMode(InConfig.EntropyCodingMode);
+	TAVResult<CFStringRef> const ConvertedEntropyCodingMode = ConvertEntropyCodingMode(InConfig.EntropyCodingMode);
 	if (ConvertedEntropyCodingMode.IsNotSuccess())
 	{
 		return ConvertedEntropyCodingMode;
 	}
 
-    OutConfig.EntropyCodingMode = ConvertedEntropyCodingMode;
+	OutConfig.EntropyCodingMode = ConvertedEntropyCodingMode;
 
 	static auto const ConvertProfile = [](EH264Profile Profile) -> TAVResult<CFStringRef>
 	{
@@ -108,7 +108,7 @@ DLLEXPORT FAVResult FAVExtension::TransformConfig(FVideoEncoderConfigVT& OutConf
 	if (ConvertedProfile.IsNotSuccess())
 	{
 		return ConvertedProfile;
-    }
+	}
 
 	OutConfig.Profile = ConvertedProfile.ReturnValue;
 
@@ -116,18 +116,18 @@ DLLEXPORT FAVResult FAVExtension::TransformConfig(FVideoEncoderConfigVT& OutConf
 }
 
 template <>
-DLLEXPORT FAVResult FAVExtension::TransformConfig(FVideoEncoderConfigVT& OutConfig, FVideoEncoderConfigH265 const& InConfig)
+DLLEXPORT FAVResult FAVExtension::TransformConfig(FVideoEncoderConfigVT &OutConfig, FVideoEncoderConfigH265 const &InConfig)
 {
 	static auto const ConvertProfile = [](EH265Profile Profile) -> TAVResult<CFStringRef>
 	{
 		switch (Profile)
 		{
 		case EH265Profile::Main:
-            return kVTProfileLevel_HEVC_Main_AutoLevel;
-        case EH265Profile::Main10:
-            return kVTProfileLevel_HEVC_Main10_AutoLevel;
-        case EH265Profile::Main422_10:
-            return kVTProfileLevel_HEVC_Main42210_AutoLevel;
+			return kVTProfileLevel_HEVC_Main_AutoLevel;
+		case EH265Profile::Main10:
+			return kVTProfileLevel_HEVC_Main10_AutoLevel;
+		case EH265Profile::Main422_10:
+			return kVTProfileLevel_HEVC_Main42210_AutoLevel;
 		default:
 			return FAVResult(EAVResult::ErrorUnsupported, FString::Printf(TEXT("H265 profile %d is not supported"), Profile), TEXT("VT"));
 		}
@@ -141,7 +141,7 @@ DLLEXPORT FAVResult FAVExtension::TransformConfig(FVideoEncoderConfigVT& OutConf
 		return ConvertedProfile;
 	}
 
-    OutConfig.Profile = ConvertedProfile.ReturnValue;
+	OutConfig.Profile = ConvertedProfile.ReturnValue;
 
 	return FAVExtension::TransformConfig<FVideoEncoderConfigVT, FVideoEncoderConfig>(OutConfig, InConfig);
 }
diff --git a/Engine/Plugins/Experimental/AVCodecs/VTCodecs/Source/VTCodecs/Public/Video/Encoders/VideoEncoderVT.hpp b/Engine/Plugins/Experimental/AVCodecs/VTCodecs/Source/VTCodecs/Public/Video/Encoders/VideoEncoderVT.hpp
--- a/Engine/Plugins/Experimental/AVCodecs/VTCodecs/Source/VTCodecs/Public/Video/Encoders/VideoEncoderVT.hpp
+++ b/Engine/Plugins/Experimental/AVCodecs/VTCodecs/Source/VTCodecs/Public/Video/Encoders/VideoEncoderVT.hpp
@@ -10,85 +10,76 @@ THIRD_PARTY_INCLUDES_START
 #include <CoreMedia/CMSync.h>
 THIRD_PARTY_INCLUDES_END
 
-#define CONDITIONAL_RELEASE(x)          \
-    if (x)                              \
-    {                                   \
-        CFRelease(x);                   \
-        x = nullptr;                    \
+#define CONDITIONAL_RELEASE(x) \
+    if (x)                     \
+    {                          \
+        CFRelease(x);          \
+        x = nullptr;           \
     }
 
 template <typename TResource>
-TVideoEncoderVT<TResource>::~TVideoEncoderVT()
+TVideoEncoderVT<TResource>::~TVideoEncoderVT() 
 {
-	Close();
+    Close();
 }
 
 template <typename TResource>
 bool TVideoEncoderVT<TResource>::IsOpen() const
 {
-	return bIsOpen;
+    return bIsOpen;
 }
 
 template <typename TResource>
 FAVResult TVideoEncoderVT<TResource>::Open(TSharedRef<FAVDevice> const& NewDevice, TSharedRef<FAVInstance> const& NewInstance)
 {
-	Close();
+    Close();
 
-	TVideoEncoder<TResource, FVideoEncoderConfigVT>::Open(NewDevice, NewInstance);
+    TVideoEncoder<TResource, FVideoEncoderConfigVT>::Open(NewDevice, NewInstance);
 
-	FrameCount = 0;
+    FrameCount = 0;
 
     bIsOpen = true;
 
-	return EAVResult::Success;
+    return EAVResult::Success;
 }
 
 template <typename TResource>
 void TVideoEncoderVT<TResource>::Close()
 {
-    if (Encoder) 
+    if (Encoder)
     {
         VTCompressionSessionInvalidate(Encoder);
         CFRelease(Encoder);
         Encoder = nullptr;
     }
-    
+
     bIsOpen = false;
 }
 
 template <typename TResource>
 bool TVideoEncoderVT<TResource>::IsInitialized() const
 {
-	return Encoder != nullptr;
+    return Encoder != nullptr;
 }
 
 template <typename TResource>
 FAVResult TVideoEncoderVT<TResource>::ApplyConfig()
 {
-	if (IsOpen())
-	{
-		FVideoEncoderConfigVT const& PendingConfig = this->GetPendingConfig();
-		if (this->AppliedConfig != PendingConfig)
-		{
-			if (IsInitialized())
-			{
-                if(this->AppliedConfig.Width == PendingConfig.Width
-                    && this->AppliedConfig.Height == PendingConfig.Height
-                    && this->AppliedConfig.Codec == PendingConfig.Codec
-                    && this->AppliedConfig.RateControlMode == PendingConfig.RateControlMode
-                    && this->AppliedConfig.KeyframeInterval == PendingConfig.KeyframeInterval
-                    && this->AppliedConfig.Profile == PendingConfig.Profile
-                    && this->AppliedConfig.PixelFormat == PendingConfig.PixelFormat
-                    && this->AppliedConfig.EntropyCodingMode == PendingConfig.EntropyCodingMode
-                    && this->AppliedConfig.MinQP == PendingConfig.MinQP
-                    && this->AppliedConfig.MaxQP == PendingConfig.MaxQP)
+    if (IsOpen())
+    {
+        FVideoEncoderConfigVT const &PendingConfig = this->GetPendingConfig();
+        if (this->AppliedConfig != PendingConfig)
+        {
+            if (IsInitialized())
+            {
+                if (this->AppliedConfig.Width == PendingConfig.Width && this->AppliedConfig.Height == PendingConfig.Height && this->AppliedConfig.Codec == PendingConfig.Codec && this->AppliedConfig.RateControlMode == PendingConfig.RateControlMode && this->AppliedConfig.KeyframeInterval == PendingConfig.KeyframeInterval && this->AppliedConfig.Profile == PendingConfig.Profile && this->AppliedConfig.PixelFormat == PendingConfig.PixelFormat && this->AppliedConfig.EntropyCodingMode == PendingConfig.EntropyCodingMode && this->AppliedConfig.MinQuality == PendingConfig.MinQuality && this->AppliedConfig.MaxQuality == PendingConfig.MaxQuality)
                 {
                     // Only bitrates can be configured on the fly
                     SetEncoderBitrate(PendingConfig);
                 }
                 else
                 {
-                    if (Encoder) 
+                    if (Encoder)
                     {
                         VTCompressionSessionInvalidate(Encoder);
                         CFRelease(Encoder);
@@ -96,50 +87,48 @@ FAVResult TVideoEncoderVT<TResource>::ApplyConfig()
                         FAVResult::Log(EAVResult::Success, TEXT("Re-initializing encoding session"), TEXT("VT"));
                     }
                 }
-			}
+            }
 
-			if (!IsInitialized())
-			{
+            if (!IsInitialized())
+            {
                 // Set source image buffer attributes. These attributes will be present on
                 // buffers retrieved from the encoder's pixel buffer pool.
                 const size_t AttributesSize = 3;
                 CFTypeRef Keys[AttributesSize] = {
                     kCVPixelBufferOpenGLCompatibilityKey,
                     kCVPixelBufferIOSurfacePropertiesKey,
-                    kCVPixelBufferPixelFormatTypeKey
-                };
+                    kCVPixelBufferPixelFormatTypeKey};
 
                 CFDictionaryRef IOSurfaceValue = CFDictionaryCreate(kCFAllocatorDefault, nullptr, nullptr, 0, &kCFTypeDictionaryKeyCallBacks, &kCFTypeDictionaryValueCallBacks);
-               
+
                 int64_t PixelType = PendingConfig.GetCVPixelFormatType();
                 CFNumberRef PixelFormat = CFNumberCreate(nullptr, kCFNumberLongType, &PixelType);
 
-                CFTypeRef Values[AttributesSize] = { kCFBooleanTrue, IOSurfaceValue, PixelFormat };
+                CFTypeRef Values[AttributesSize] = {kCFBooleanTrue, IOSurfaceValue, PixelFormat};
                 CFDictionaryRef SourceAttributes = CFDictionaryCreate(kCFAllocatorDefault, Keys, Values, AttributesSize, &kCFTypeDictionaryKeyCallBacks, &kCFTypeDictionaryValueCallBacks);
 
                 CONDITIONAL_RELEASE(IOSurfaceValue);
                 CONDITIONAL_RELEASE(PixelFormat);
-				
-				// Encoder specifications
-				const size_t EncoderSpecsSize = 1;
-				CFTypeRef EncoderKeys[EncoderSpecsSize] =
-				{
-					// We want HW acceleration for best latency, if we can't get it then fail creating the session
-					kVTVideoEncoderSpecification_RequireHardwareAcceleratedVideoEncoder,
-				};
-				
-				CFTypeRef EncoderValues[EncoderSpecsSize] =
-				{
-					kCFBooleanTrue
-				};
-				
-				CFDictionaryRef EncoderSpec = CFDictionaryCreate(kCFAllocatorDefault, EncoderKeys, EncoderValues, EncoderSpecsSize, &kCFTypeDictionaryKeyCallBacks, &kCFTypeDictionaryValueCallBacks);
-          
-                OSStatus Result = VTCompressionSessionCreate(kCFAllocatorDefault, 
+
+                // Encoder specifications
+                const size_t EncoderSpecsSize = 1;
+                CFTypeRef EncoderKeys[EncoderSpecsSize] =
+                    {
+                        // We want HW acceleration for best latency, if we can't get it then fail creating the session
+                        kVTVideoEncoderSpecification_RequireHardwareAcceleratedVideoEncoder,
+                    };
+
+                CFTypeRef EncoderValues[EncoderSpecsSize] =
+                    {
+                        kCFBooleanTrue};
+
+                CFDictionaryRef EncoderSpec = CFDictionaryCreate(kCFAllocatorDefault, EncoderKeys, EncoderValues, EncoderSpecsSize, &kCFTypeDictionaryKeyCallBacks, &kCFTypeDictionaryValueCallBacks);
+
+                OSStatus Result = VTCompressionSessionCreate(kCFAllocatorDefault,
                                                              PendingConfig.Width,
                                                              PendingConfig.Height,
                                                              PendingConfig.Codec,
-															 EncoderSpec,
+                                                             EncoderSpec,
                                                              SourceAttributes,
                                                              NULL /* Default Compressed Data Allocator */,
                                                              Internal::VTCompressionOutputCallback,
@@ -147,11 +136,11 @@ FAVResult TVideoEncoderVT<TResource>::ApplyConfig()
                                                              &Encoder);
 
                 CONDITIONAL_RELEASE(SourceAttributes);
-				CONDITIONAL_RELEASE(EncoderSpec);
-                    
-                if(Result != 0)
+                CONDITIONAL_RELEASE(EncoderSpec);
+
+                if (Result != 0)
                 {
-                    if(Result == kVTInvalidSessionErr)
+                    if (Result == kVTInvalidSessionErr)
                     {
                         FAVResult::Log(EAVResult::Warning, TEXT("Invalid VTCompressionSession. You could be exceeding the maximum resolution supported by VideoToolbox!"), TEXT("VT"));
                     }
@@ -159,39 +148,41 @@ FAVResult TVideoEncoderVT<TResource>::ApplyConfig()
                 }
 
                 ConfigureCompressionSession(PendingConfig);
-			}
-		}
+            }
+        }
 
-		return TVideoEncoder<TResource, FVideoEncoderConfigVT>::ApplyConfig();
-	}
+        return TVideoEncoder<TResource, FVideoEncoderConfigVT>::ApplyConfig();
+    }
 
-	return FAVResult(EAVResult::ErrorInvalidState, TEXT("Encoder not open"), TEXT("VT"));
+    return FAVResult(EAVResult::ErrorInvalidState, TEXT("Encoder not open"), TEXT("VT"));
 }
 
 template <typename TResource>
-FAVResult TVideoEncoderVT<TResource>::ConfigureCompressionSession(FVideoEncoderConfigVT const& Config)  
-{       
+FAVResult TVideoEncoderVT<TResource>::ConfigureCompressionSession(FVideoEncoderConfigVT const& Config)
+{
     VTSessionHelpers::SetVTSessionProperty(Encoder, kVTCompressionPropertyKey_RealTime, true);
 
     VTSessionHelpers::SetVTSessionProperty(Encoder, kVTCompressionPropertyKey_ProfileLevel, Config.Profile);
-    
+
     SetEncoderBitrate(Config);
-    
-	// Enable temporal compression, e.g. creating delta frames
+
+    // Enable temporal compression, e.g. creating delta frames
     VTSessionHelpers::SetVTSessionProperty(Encoder, kVTCompressionPropertyKey_AllowTemporalCompression, true);
-    
-	// Frame reordering is unhelpful in WebRTC, real-time streaming usecases
+
+    // Frame reordering is unhelpful in WebRTC, real-time streaming usecases
     VTSessionHelpers::SetVTSessionProperty(Encoder, kVTCompressionPropertyKey_AllowFrameReordering, false);
-    
+
     VTSessionHelpers::SetVTSessionProperty(Encoder, kVTCompressionPropertyKey_MaxKeyFrameInterval, (int32_t)Config.KeyframeInterval);
-    
+
     VTSessionHelpers::SetVTSessionProperty(Encoder, kVTCompressionPropertyKey_MaxKeyFrameIntervalDuration, (Config.KeyframeInterval > 0 ? (int32_t)(Config.KeyframeInterval / Config.FrameRate) : 0));
-    
-    VTSessionHelpers::SetVTSessionProperty(Encoder, kVTCompressionPropertyKey_MinAllowedFrameQP, (int32_t)Config.MinQP);
-	
-	VTSessionHelpers::SetVTSessionProperty(Encoder, kVTCompressionPropertyKey_MaxAllowedFrameQP, (int32_t)Config.MaxQP);
 
-    if(Config.Codec == kCMVideoCodecType_H264)
+    int32_t MinQP = (1.0f - (Config.MaxQuality / 100.0f)) * 51.0f;
+    VTSessionHelpers::SetVTSessionProperty(Encoder, kVTCompressionPropertyKey_MinAllowedFrameQP, MinQP);
+
+    int32_t MaxQP = (1.0f - (Config.MinQuality / 100.0f)) * 51.0f;
+    VTSessionHelpers::SetVTSessionProperty(Encoder, kVTCompressionPropertyKey_MaxAllowedFrameQP, MaxQP);
+
+    if (Config.Codec == kCMVideoCodecType_H264)
     {
         VTSessionHelpers::SetVTSessionProperty(Encoder, kVTCompressionPropertyKey_H264EntropyMode, Config.EntropyCodingMode);
     }
@@ -200,19 +191,19 @@ FAVResult TVideoEncoderVT<TResource>::ConfigureCompressionSession(FVideoEncoderC
 }
 
 template <typename TResource>
-FAVResult TVideoEncoderVT<TResource>::SetEncoderBitrate(FVideoEncoderConfigVT const& Config)
+FAVResult TVideoEncoderVT<TResource>::SetEncoderBitrate(FVideoEncoderConfigVT const &Config)
 {
     VTSessionHelpers::SetVTSessionProperty(Encoder, kVTCompressionPropertyKey_AverageBitRate, Config.TargetBitrate);
-    
+
     return EAVResult::Success;
 }
 
 template <typename TResource>
-FAVResult TVideoEncoderVT<TResource>::SendFrame(TSharedPtr<FVideoResourceMetal> const& Resource, uint32 Timestamp, bool bForceKeyframe)
+FAVResult TVideoEncoderVT<TResource>::SendFrame(TSharedPtr<FVideoResourceMetal> const &Resource, uint32 Timestamp, bool bForceKeyframe)
 {
-    if(IsOpen())
+    if (IsOpen())
     {
-        FVideoEncoderConfigVT& PendingConfig = this->EditPendingConfig();
+        FVideoEncoderConfigVT &PendingConfig = this->EditPendingConfig();
         PendingConfig.PixelFormat = Resource->GetFormat();
 
         FAVResult AVResult = ApplyConfig();
@@ -221,7 +212,7 @@ FAVResult TVideoEncoderVT<TResource>::SendFrame(TSharedPtr<FVideoResourceMetal>
             return AVResult;
         }
 
-        if(Resource.IsValid())
+        if (Resource.IsValid())
         {
             float TimestampSecs = Timestamp / 1000000.0f;
             CMTime PresentationTime = CMTimeMakeWithSeconds(TimestampSecs, NSEC_PER_SEC);
@@ -236,22 +227,22 @@ FAVResult TVideoEncoderVT<TResource>::SendFrame(TSharedPtr<FVideoResourceMetal>
             TUniquePtr<EncodeParams> Params = MakeUnique<EncodeParams>();
             Params.Reset(new EncodeParams(PendingConfig.Codec, PresentationTime));
 
-            OSStatus Status = VTCompressionSessionEncodeFrame(Encoder, Resource->GetRaw(), PresentationTime, kCMTimeInvalid, FrameProperties, (void*)Params.Release(), NULL);
+            OSStatus Status = VTCompressionSessionEncodeFrame(Encoder, Resource->GetRaw(), PresentationTime, kCMTimeInvalid, FrameProperties, (void *)Params.Release(), NULL);
 
             CONDITIONAL_RELEASE(FrameProperties);
 
-            if(Status != kCVReturnSuccess)
+            if (Status != kCVReturnSuccess)
             {
                 return FAVResult(EAVResult::Error, TEXT("Failed to encode"), TEXT("VT"), Status);
             }
-            
+
             return EAVResult::Success;
         }
         else
         {
             // Flush encoder
             OSStatus Status = VTCompressionSessionCompleteFrames(Encoder, kCMTimeInvalid);
-            if(Status != kCVReturnSuccess)
+            if (Status != kCVReturnSuccess)
             {
                 return FAVResult(EAVResult::Error, TEXT("Failed to flush"), TEXT("VT"), Status);
             }
@@ -264,59 +255,59 @@ FAVResult TVideoEncoderVT<TResource>::SendFrame(TSharedPtr<FVideoResourceMetal>
 }
 
 template <typename TResource>
-FAVResult TVideoEncoderVT<TResource>::ReceivePacket(FVideoPacket& OutPacket)
+FAVResult TVideoEncoderVT<TResource>::ReceivePacket(FVideoPacket &OutPacket)
 {
-	if (IsOpen())
-	{
-		if (Packets.Dequeue(OutPacket))
-		{
-			return EAVResult::Success;
-		}
+    if (IsOpen())
+    {
+        if (Packets.Dequeue(OutPacket))
+        {
+            return EAVResult::Success;
+        }
 
-		return EAVResult::PendingInput;
-	}
+        return EAVResult::PendingInput;
+    }
 
-	return FAVResult(EAVResult::ErrorInvalidState, TEXT("Encoder not open"), TEXT("VT"));
+    return FAVResult(EAVResult::ErrorInvalidState, TEXT("Encoder not open"), TEXT("VT"));
 }
 
 template <typename TResource>
-FAVResult TVideoEncoderVT<TResource>::HandlePacket(void* Params, OSStatus Status, VTEncodeInfoFlags InfoFlags, CMSampleBufferRef& SampleBuffer)
+FAVResult TVideoEncoderVT<TResource>::HandlePacket(void *Params, OSStatus Status, VTEncodeInfoFlags InfoFlags, CMSampleBufferRef &SampleBuffer)
 {
-	if (IsOpen())
-	{
-        if(Status != 0)
+    if (IsOpen())
+    {
+        if (Status != 0)
         {
             return FAVResult(EAVResult::Error, TEXT("Failed to encode"), TEXT("VT"), Status);
         }
 
-        if(InfoFlags & kVTEncodeInfo_FrameDropped)
+        if (InfoFlags & kVTEncodeInfo_FrameDropped)
         {
             return FAVResult(EAVResult::Error, TEXT("Frame dropped"), TEXT("VT"), Status);
         }
 
         bool bIsKeyframe = false;
         CFArrayRef Attachments = CMSampleBufferGetSampleAttachmentsArray(SampleBuffer, 0);
-        if (Attachments != nullptr && CFArrayGetCount(Attachments)) 
+        if (Attachments != nullptr && CFArrayGetCount(Attachments))
         {
             CFDictionaryRef Attachment = static_cast<CFDictionaryRef>(CFArrayGetValueAtIndex(Attachments, 0));
             bIsKeyframe = !CFDictionaryContainsKey(Attachment, kCMSampleAttachmentKey_NotSync);
         }
-        
+
         // SampleBuffer is an MPEG4 container stream, but we need to convert it to its raw bitstream
-        EncodeParams* Config = static_cast<EncodeParams*>(Params);
+        EncodeParams *Config = static_cast<EncodeParams *>(Params);
 
         int QP = 0;
         TArray<uint8> Bitstream;
-        if(Config->Codec == kCMVideoCodecType_H264) 
+        if (Config->Codec == kCMVideoCodecType_H264)
         {
-            if(!NaluRewriter::H264CMSampleBufferToAnnexBBuffer(SampleBuffer, bIsKeyframe, Bitstream))
+            if (!NaluRewriter::H264CMSampleBufferToAnnexBBuffer(SampleBuffer, bIsKeyframe, Bitstream))
             {
                 return FAVResult(EAVResult::Error, TEXT("Failed to extract bitstream"), TEXT("VT"));
             }
         }
-        else if(Config->Codec == kCMVideoCodecType_HEVC) 
+        else if (Config->Codec == kCMVideoCodecType_HEVC)
         {
-            if(!NaluRewriter::H265CMSampleBufferToAnnexBBuffer(SampleBuffer, bIsKeyframe, Bitstream))
+            if (!NaluRewriter::H265CMSampleBufferToAnnexBBuffer(SampleBuffer, bIsKeyframe, Bitstream))
             {
                 return FAVResult(EAVResult::Error, TEXT("Failed to extract bitstream"), TEXT("VT"));
             }
@@ -327,26 +318,26 @@ FAVResult TVideoEncoderVT<TResource>::HandlePacket(void* Params, OSStatus Status
         }
 
         TSharedPtr<uint8> const CopiedData = MakeShareable(new uint8[Bitstream.Num()]);
-		FMemory::BigBlockMemcpy(CopiedData.Get(), Bitstream.GetData(), Bitstream.Num());
+        FMemory::BigBlockMemcpy(CopiedData.Get(), Bitstream.GetData(), Bitstream.Num());
 
-        if(Config->Codec == kCMVideoCodecType_H264)
+        if (Config->Codec == kCMVideoCodecType_H264)
         {
             using namespace UE::AVCodecCore::H264;
-			TArray<Slice_t> Slices;
+            TArray<Slice_t> Slices;
             FAVResult Result = H264.Parse(FVideoPacket(CopiedData, Bitstream.Num(), 0, 0, 0, false), Slices);
-            if(Result != EAVResult::Success)
+            if (Result != EAVResult::Success)
             {
                 return FAVResult(EAVResult::Error, TEXT("Failed to parse bitstream"), TEXT("VT"), Result);
             }
 
             QP = H264.GetLastSliceQP(Slices).Get(0);
         }
-        else if(Config->Codec == kCMVideoCodecType_HEVC)
+        else if (Config->Codec == kCMVideoCodecType_HEVC)
         {
             using namespace UE::AVCodecCore::H265;
-			TArray<Slice_t> Slices;
+            TArray<Slice_t> Slices;
             FAVResult Result = H265.Parse(FVideoPacket(CopiedData, Bitstream.Num(), 0, 0, 0, false), Slices);
-            if(Result != EAVResult::Success)
+            if (Result != EAVResult::Success)
             {
                 return FAVResult(EAVResult::Error, TEXT("Failed to parse bitstream"), TEXT("VT"), Result);
             }
@@ -355,8 +346,8 @@ FAVResult TVideoEncoderVT<TResource>::HandlePacket(void* Params, OSStatus Status
             FAVResult::Log(EAVResult::Warning, TEXT("H265 QP"), TEXT("VT"), QP);
         }
 
-		return Packets.Enqueue(FVideoPacket(CopiedData, Bitstream.Num(), CMClockConvertHostTimeToSystemUnits(Config->Timestamp), ++FrameCount, (uint32_t)QP, bIsKeyframe)) ? EAVResult::Success : EAVResult::Error;
-	}
+        return Packets.Enqueue(FVideoPacket(CopiedData, Bitstream.Num(), CMClockConvertHostTimeToSystemUnits(Config->Timestamp), ++FrameCount, (uint32_t)QP, bIsKeyframe)) ? EAVResult::Success : EAVResult::Error;
+    }
 
-	return FAVResult(EAVResult::ErrorInvalidState, TEXT("Encoder not open"), TEXT("VT"));
+    return FAVResult(EAVResult::ErrorInvalidState, TEXT("Encoder not open"), TEXT("VT"));
 }
