Fixing remaining MinQP, MaxQP references in VTCodecs.

#okforgit public

diff --git a/Engine/Plugins/Experimental/AVCodecs/VTCodecs/Source/VTCodecs/Private/Video/Encoders/Configs/VideoEncoderConfigVT.cpp b/Engine/Plugins/Experimental/AVCodecs/VTCodecs/Source/VTCodecs/Private/Video/Encoders/Configs/VideoEncoderConfigVT.cpp
--- a/Engine/Plugins/Experimental/AVCodecs/VTCodecs/Source/VTCodecs/Private/Video/Encoders/Configs/VideoEncoderConfigVT.cpp
+++ b/Engine/Plugins/Experimental/AVCodecs/VTCodecs/Source/VTCodecs/Private/Video/Encoders/Configs/VideoEncoderConfigVT.cpp
@@ -10,7 +10,7 @@ REGISTER_TYPEID(FVideoEncoderConfigVT);
 static uint32 DEFAULT_BITRATE_TARGET = 1000000;
 static uint32 DEFAULT_BITRATE_MAX = 10000000;
 
-TAVResult<OSType> FVideoEncoderConfigVT::ConvertFormat(EVideoFormat const &Format)
+TAVResult<OSType> FVideoEncoderConfigVT::ConvertFormat(EVideoFormat const& Format)
 {
 	switch (Format)
 	{
@@ -26,7 +26,7 @@ TAVResult<OSType> FVideoEncoderConfigVT::ConvertFormat(EVideoFormat const &Forma
 }
 
 template <>
-DLLEXPORT FAVResult FAVExtension::TransformConfig(FVideoEncoderConfigVT &OutConfig, FVideoEncoderConfig const &InConfig)
+DLLEXPORT FAVResult FAVExtension::TransformConfig(FVideoEncoderConfigVT& OutConfig, FVideoEncoderConfig const& InConfig)
 {
 	OutConfig.Width = InConfig.Width;
 	OutConfig.Height = InConfig.Height;
@@ -43,7 +43,7 @@ DLLEXPORT FAVResult FAVExtension::TransformConfig(FVideoEncoderConfigVT &OutConf
 }
 
 template <>
-DLLEXPORT FAVResult FAVExtension::TransformConfig(FVideoEncoderConfig &OutConfig, FVideoEncoderConfigVT const &InConfig)
+DLLEXPORT FAVResult FAVExtension::TransformConfig(FVideoEncoderConfig& OutConfig, FVideoEncoderConfigVT const& InConfig)
 {
 	OutConfig.Width = InConfig.Width;
 	OutConfig.Height = InConfig.Height;
@@ -60,24 +60,24 @@ DLLEXPORT FAVResult FAVExtension::TransformConfig(FVideoEncoderConfig &OutConfig
 }
 
 template <>
-DLLEXPORT FAVResult FAVExtension::TransformConfig(FVideoEncoderConfigVT &OutConfig, FVideoEncoderConfigH264 const &InConfig)
+DLLEXPORT FAVResult FAVExtension::TransformConfig(FVideoEncoderConfigVT& OutConfig, FVideoEncoderConfigH264 const& InConfig)
 {
 	OutConfig.Codec = kCMVideoCodecType_H264;
 
 	static auto const ConvertEntropyCodingMode = [](EH264EntropyCodingMode EntropyCodingMode) -> TAVResult<CFStringRef>
-	{
-		switch (EntropyCodingMode)
 		{
-		case EH264EntropyCodingMode::Auto:
-			return nullptr;
-		case EH264EntropyCodingMode::CABAC:
-			return kVTH264EntropyMode_CABAC;
-		case EH264EntropyCodingMode::CAVLC:
-			return kVTH264EntropyMode_CAVLC;
-		default:
-			return FAVResult(EAVResult::ErrorUnsupported, FString::Printf(TEXT("H264 entropy coding mode %d is not supported"), EntropyCodingMode), TEXT("VT"));
-		}
-	};
+			switch (EntropyCodingMode)
+			{
+			case EH264EntropyCodingMode::Auto:
+				return nullptr;
+			case EH264EntropyCodingMode::CABAC:
+				return kVTH264EntropyMode_CABAC;
+			case EH264EntropyCodingMode::CAVLC:
+				return kVTH264EntropyMode_CAVLC;
+			default:
+				return FAVResult(EAVResult::ErrorUnsupported, FString::Printf(TEXT("H264 entropy coding mode %d is not supported"), EntropyCodingMode), TEXT("VT"));
+			}
+		};
 
 	TAVResult<CFStringRef> const ConvertedEntropyCodingMode = ConvertEntropyCodingMode(InConfig.EntropyCodingMode);
 	if (ConvertedEntropyCodingMode.IsNotSuccess())
@@ -88,21 +88,21 @@ DLLEXPORT FAVResult FAVExtension::TransformConfig(FVideoEncoderConfigVT &OutConf
 	OutConfig.EntropyCodingMode = ConvertedEntropyCodingMode;
 
 	static auto const ConvertProfile = [](EH264Profile Profile) -> TAVResult<CFStringRef>
-	{
-		switch (Profile)
 		{
-		case EH264Profile::Baseline:
-			return kVTProfileLevel_H264_Baseline_AutoLevel;
-		case EH264Profile::Main:
-			return kVTProfileLevel_H264_Main_AutoLevel;
-		case EH264Profile::High:
-			return kVTProfileLevel_H264_High_AutoLevel;
-		case EH264Profile::ConstrainedHigh:
-			return kVTProfileLevel_H264_ConstrainedHigh_AutoLevel;
-		default:
-			return FAVResult(EAVResult::ErrorUnsupported, FString::Printf(TEXT("H264 profile %d is not supported"), Profile), TEXT("VT"));
-		}
-	};
+			switch (Profile)
+			{
+			case EH264Profile::Baseline:
+				return kVTProfileLevel_H264_Baseline_AutoLevel;
+			case EH264Profile::Main:
+				return kVTProfileLevel_H264_Main_AutoLevel;
+			case EH264Profile::High:
+				return kVTProfileLevel_H264_High_AutoLevel;
+			case EH264Profile::ConstrainedHigh:
+				return kVTProfileLevel_H264_ConstrainedHigh_AutoLevel;
+			default:
+				return FAVResult(EAVResult::ErrorUnsupported, FString::Printf(TEXT("H264 profile %d is not supported"), Profile), TEXT("VT"));
+			}
+		};
 
 	TAVResult<CFStringRef> const ConvertedProfile = ConvertProfile(InConfig.Profile);
 	if (ConvertedProfile.IsNotSuccess())
@@ -116,22 +116,22 @@ DLLEXPORT FAVResult FAVExtension::TransformConfig(FVideoEncoderConfigVT &OutConf
 }
 
 template <>
-DLLEXPORT FAVResult FAVExtension::TransformConfig(FVideoEncoderConfigVT &OutConfig, FVideoEncoderConfigH265 const &InConfig)
+DLLEXPORT FAVResult FAVExtension::TransformConfig(FVideoEncoderConfigVT& OutConfig, FVideoEncoderConfigH265 const& InConfig)
 {
 	static auto const ConvertProfile = [](EH265Profile Profile) -> TAVResult<CFStringRef>
-	{
-		switch (Profile)
 		{
-		case EH265Profile::Main:
-			return kVTProfileLevel_HEVC_Main_AutoLevel;
-		case EH265Profile::Main10:
-			return kVTProfileLevel_HEVC_Main10_AutoLevel;
-		case EH265Profile::Main422_10:
-			return kVTProfileLevel_HEVC_Main42210_AutoLevel;
-		default:
-			return FAVResult(EAVResult::ErrorUnsupported, FString::Printf(TEXT("H265 profile %d is not supported"), Profile), TEXT("VT"));
-		}
-	};
+			switch (Profile)
+			{
+			case EH265Profile::Main:
+				return kVTProfileLevel_HEVC_Main_AutoLevel;
+			case EH265Profile::Main10:
+				return kVTProfileLevel_HEVC_Main10_AutoLevel;
+			case EH265Profile::Main422_10:
+				return kVTProfileLevel_HEVC_Main42210_AutoLevel;
+			default:
+				return FAVResult(EAVResult::ErrorUnsupported, FString::Printf(TEXT("H265 profile %d is not supported"), Profile), TEXT("VT"));
+			}
+		};
 
 	OutConfig.Codec = kCMVideoCodecType_HEVC;
 
diff --git a/Engine/Plugins/Experimental/AVCodecs/VTCodecs/Source/VTCodecs/Public/Video/Encoders/Configs/VideoEncoderConfigVT.h b/Engine/Plugins/Experimental/AVCodecs/VTCodecs/Source/VTCodecs/Public/Video/Encoders/Configs/VideoEncoderConfigVT.h
--- a/Engine/Plugins/Experimental/AVCodecs/VTCodecs/Source/VTCodecs/Public/Video/Encoders/Configs/VideoEncoderConfigVT.h
+++ b/Engine/Plugins/Experimental/AVCodecs/VTCodecs/Source/VTCodecs/Public/Video/Encoders/Configs/VideoEncoderConfigVT.h
@@ -28,8 +28,8 @@ public:
 	CFStringRef Profile;
     EVideoFormat PixelFormat;
     CFStringRef EntropyCodingMode;
-    int32 MinQP;
-    int32 MaxQP;
+    int32 MinQuality;
+    int32 MaxQuality;
 
 	FVideoEncoderConfigVT()
 		: FAVConfig()
@@ -65,8 +65,8 @@ public:
 				this->TargetBitrate == Other.TargetBitrate &&
                 this->PixelFormat == Other.PixelFormat &&
                 this->EntropyCodingMode == Other.EntropyCodingMode &&
-                this->MinQP == Other.MinQP &&
-                this->MaxQP == Other.MaxQP;
+                this->MinQuality == Other.MinQuality &&
+                this->MaxQuality == Other.MaxQuality;
 	}
 
 	bool operator!=(FVideoEncoderConfigVT const& Other) const
