Fix Xcode16.3 compile issues
- WebSocketServer iface assignment goes out of scope, leading to dangling pointer, so keep an instance backed copy
- TextureCache in VideoDecoderVT.hpp shadows parents instance var

#rnx

diff --git a/Engine/Plugins/Experimental/AVCodecs/VTCodecs/Source/VTCodecs/Public/Video/Decoders/VideoDecoderVT.hpp b/Engine/Plugins/Experimental/AVCodecs/VTCodecs/Source/VTCodecs/Public/Video/Decoders/VideoDecoderVT.hpp
--- a/Engine/Plugins/Experimental/AVCodecs/VTCodecs/Source/VTCodecs/Public/Video/Decoders/VideoDecoderVT.hpp
+++ b/Engine/Plugins/Experimental/AVCodecs/VTCodecs/Source/VTCodecs/Public/Video/Decoders/VideoDecoderVT.hpp
@@ -252,13 +252,13 @@ FAVResult TVideoDecoderVT<TResource>::ReceiveFrame(TResolvableVideoResource<TRes
             size_t Height = CVPixelBufferGetHeight(Frame->ImageBuffer);
             OSType PixelFormat = CVPixelBufferGetPixelFormatType(Frame->ImageBuffer);
             
-            static CVMetalTextureCacheRef TextureCache = nullptr;
-            if (TextureCache == nullptr)
+            static CVMetalTextureCacheRef MetalTextureCache = nullptr;
+            if (MetalTextureCache == nullptr)
             {
                 id<MTLDevice> Device = (__bridge id<MTLDevice>)GDynamicRHI->RHIGetNativeDevice();
                 check(Device);
                 
-                CVReturn Result = CVMetalTextureCacheCreate(kCFAllocatorDefault, nullptr, Device, nullptr, &TextureCache);
+                CVReturn Result = CVMetalTextureCacheCreate(kCFAllocatorDefault, nullptr, Device, nullptr, &MetalTextureCache);
                 if (Result != kCVReturnSuccess)
                 {
                     return FAVResult(EAVResult::Error, TEXT("Failed to create CVMetalTextureCacheRef"), TEXT("VT"), Result);
@@ -266,7 +266,7 @@ FAVResult TVideoDecoderVT<TResource>::ReceiveFrame(TResolvableVideoResource<TRes
             }
             
             CVMetalTextureRef TextureRef = nullptr;
-            CVReturn Result = CVMetalTextureCacheCreateTextureFromImage(kCFAllocatorDefault, TextureCache, Frame->ImageBuffer, nullptr, MTLPixelFormatBGRA8Unorm_sRGB, Width, Height, 0, &TextureRef);
+            CVReturn Result = CVMetalTextureCacheCreateTextureFromImage(kCFAllocatorDefault, MetalTextureCache, Frame->ImageBuffer, nullptr, MTLPixelFormatBGRA8Unorm_sRGB, Width, Height, 0, &TextureRef);
             if (Result != kCVReturnSuccess)
             {
                 return FAVResult(EAVResult::Error, TEXT("Failed to create CVMetalTextureRef"), TEXT("VT"), Result);
diff --git a/Engine/Plugins/Experimental/WebSocketNetworking/Source/WebSocketNetworking/Private/WebSocketServer.cpp b/Engine/Plugins/Experimental/WebSocketNetworking/Source/WebSocketNetworking/Private/WebSocketServer.cpp
--- a/Engine/Plugins/Experimental/WebSocketNetworking/Source/WebSocketNetworking/Private/WebSocketServer.cpp
+++ b/Engine/Plugins/Experimental/WebSocketNetworking/Source/WebSocketNetworking/Private/WebSocketServer.cpp
@@ -109,7 +109,7 @@ void FWebSocketServer::EnableHTTPServer(TArray<FWebSocketHttpMount> InDirectorie
 #endif
 }
 
-bool FWebSocketServer::Init(uint32 Port, FWebSocketClientConnectedCallBack CallBack, FString BindAddress)
+bool FWebSocketServer::Init(uint32 Port, FWebSocketClientConnectedCallBack CallBack, FString InBindAddress)
 {
 #if USE_LIBWEBSOCKET
 	#if !UE_BUILD_SHIPPING
@@ -141,9 +141,12 @@ bool FWebSocketServer::Init(uint32 Port, FWebSocketClientConnectedCallBack CallB
 	ServerPort = Port;
 
 	Info.iface = NULL;
-	if (!BindAddress.IsEmpty())
+	if (!InBindAddress.IsEmpty())
 	{
-		Info.iface = StringCast<ANSICHAR>(*BindAddress).Get();
+		size_t AddressLen = FCStringAnsi::Strlen(StringCast<ANSICHAR>(*InBindAddress).Get());
+		BindAddress = new char[AddressLen + 1];
+		FCStringAnsi::Strncpy(BindAddress, StringCast<ANSICHAR>(*InBindAddress).Get(), AddressLen+1);
+		Info.iface = BindAddress;
 	}
 
 	Info.protocols = &Protocols[0];
@@ -215,6 +218,9 @@ FWebSocketServer::~FWebSocketServer()
 
 	delete[] LwsHttpMounts;
 	LwsHttpMounts = NULL;
+
+	delete[] BindAddress;
+	BindAddress = NULL;
 #endif
 }
 
diff --git a/Engine/Plugins/Experimental/WebSocketNetworking/Source/WebSocketNetworking/Private/WebSocketServer.h b/Engine/Plugins/Experimental/WebSocketNetworking/Source/WebSocketNetworking/Private/WebSocketServer.h
--- a/Engine/Plugins/Experimental/WebSocketNetworking/Source/WebSocketNetworking/Private/WebSocketServer.h
+++ b/Engine/Plugins/Experimental/WebSocketNetworking/Source/WebSocketNetworking/Private/WebSocketServer.h
@@ -48,4 +48,6 @@ private:
 
 	TArray<FWebSocketHttpMount> DirectoriesToServe;
 	WebSocketInternalHttpMount* LwsHttpMounts = NULL;
+
+	char* BindAddress = NULL;
 };
