[AVCodecs] Fix crash when Decoding VP9 on Apple with VTCodecs.

Problem is the Apple Ref counted objects were copied without incrementing the ref count.
Early exit if CreateVP9VideoFormatDescription frame size is zero.
Fixed InputFormat not released if ApplyConfig failes.

diff --git a/Engine/Plugins/Experimental/AVCodecs/VTCodecs/Source/VTCodecs/Private/Video/Util/NaluRewriter.cpp b/Engine/Plugins/Experimental/AVCodecs/VTCodecs/Source/VTCodecs/Private/Video/Util/NaluRewriter.cpp
--- a/Engine/Plugins/Experimental/AVCodecs/VTCodecs/Source/VTCodecs/Private/Video/Util/NaluRewriter.cpp
+++ b/Engine/Plugins/Experimental/AVCodecs/VTCodecs/Source/VTCodecs/Private/Video/Util/NaluRewriter.cpp
@@ -548,11 +548,16 @@ CMVideoFormatDescriptionRef NaluRewriter::CreateVP9VideoFormatDescription(const
 	FMemory::BigBlockMemcpy(CopiedData.Get(), rtp_buffer, rtp_buffer_size);
 
     VP9::Header_t Header;
-    if(VP9.Parse(FVideoPacket(CopiedData, rtp_buffer_size, 0, 0, 0, false), Header) != EAVResult::Success)
+	if (VP9.Parse(FVideoPacket(CopiedData, rtp_buffer_size, 0, 0, 0, false), Header) != EAVResult::Success)
     {
-        FAVResult::Log(EAVResult::Error, TEXT("CreateVP9VideoFormatDescription failed to parse VP9 bitstream"), TEXT("VT"));
+		FAVResult::Log(EAVResult::Error, TEXT("CreateVP9VideoFormatDescription failed to parse VP9 bitstream"), TEXT("VT"));
         return nullptr;
     }
+	// Early exit if size is not set because this means a valid VP9VideoFormatDescription cannot be constructed.
+	if (Header.frame_width_minus_1 <= 0)
+	{
+		return nullptr;
+	}
 
     const size_t VPCodecConfigurationContentsSize = 12;
 
@@ -561,15 +566,19 @@ CMVideoFormatDescriptionRef NaluRewriter::CreateVP9VideoFormatDescription(const
     {
         case VP9::ESubSampling::k444:
             ChromaSubSampling = 3;
+			break;
         case VP9::ESubSampling::k440:
             ChromaSubSampling = 1;
+			break;
         case VP9::ESubSampling::k422:
             ChromaSubSampling = 2;
+			break;
         case VP9::ESubSampling::k420:
             ChromaSubSampling = 1;
+			break;
     }
 
-    uint8 BitDepthChromaAndRange = (0xF & (uint8)Header.bit_depth.Value) << 4 | (0x7 & ChromaSubSampling) << 1 | (0x1 & (uint8)Header.color_range.Value);
+	uint8 BitDepthChromaAndRange = (0xF & static_cast<uint8>(Header.bit_depth.Value)) << 4 | (0x7 & ChromaSubSampling) << 1 | (0x1 & static_cast<uint8>(Header.color_range.Value));
     
     uint8 Record[VPCodecConfigurationContentsSize];
     memset((void*)Record, 0, VPCodecConfigurationContentsSize);
@@ -601,7 +610,7 @@ CMVideoFormatDescriptionRef NaluRewriter::CreateVP9VideoFormatDescription(const
     CONDITIONAL_RELEASE(AvcInfo);
 
     CMVideoFormatDescriptionRef FormatDescription = nullptr;
-    OSStatus Status = CMVideoFormatDescriptionCreate(kCFAllocatorDefault, kCMVideoCodecType_VP9, Header.frame_width_minus_1 + 1, Header.frame_height_minus_1 + 1, ConfigInfo, &FormatDescription);
+	OSStatus Status = CMVideoFormatDescriptionCreate(kCFAllocatorDefault, kCMVideoCodecType_VP9, Header.frame_width_minus_1 + 1, Header.frame_height_minus_1 + 1, ConfigInfo, &FormatDescription);
     if (Status != noErr)
     { 
         FAVResult::Log(EAVResult::Error, TEXT("CreateVP9VideoFormatDescription failed to create video format description"), TEXT("VT"), Status);
diff --git a/Engine/Plugins/Experimental/AVCodecs/VTCodecs/Source/VTCodecs/Public/Video/Decoders/Configs/VideoDecoderConfigVT.h b/Engine/Plugins/Experimental/AVCodecs/VTCodecs/Source/VTCodecs/Public/Video/Decoders/Configs/VideoDecoderConfigVT.h
--- a/Engine/Plugins/Experimental/AVCodecs/VTCodecs/Source/VTCodecs/Public/Video/Decoders/Configs/VideoDecoderConfigVT.h
+++ b/Engine/Plugins/Experimental/AVCodecs/VTCodecs/Source/VTCodecs/Public/Video/Decoders/Configs/VideoDecoderConfigVT.h
@@ -24,10 +24,19 @@ public:
 	{
 
 	}
+
+	~FVideoDecoderConfigVT()
+	{
+		SetVideoFormat(nullptr);
+	}
 		
 	FVideoDecoderConfigVT(FVideoDecoderConfigVT const& Other)
 	{
 		*this = Other;
+		if (VideoFormat)
+		{
+			CFRetain(VideoFormat);
+		}
 	}
     
     void SetVideoFormat(CMVideoFormatDescriptionRef Format)
@@ -53,6 +62,10 @@ public:
 	FVideoDecoderConfigVT& operator=(FVideoDecoderConfigVT const& Other)
 	{
 		FMemory::Memcpy(*this, Other);
+		if (VideoFormat)
+		{
+			CFRetain(VideoFormat);
+		}
 
 		return *this;
 	}
diff --git a/Engine/Plugins/Experimental/AVCodecs/VTCodecs/Source/VTCodecs/Public/Video/Decoders/VideoDecoderVT.hpp b/Engine/Plugins/Experimental/AVCodecs/VTCodecs/Source/VTCodecs/Public/Video/Decoders/VideoDecoderVT.hpp
--- a/Engine/Plugins/Experimental/AVCodecs/VTCodecs/Source/VTCodecs/Public/Video/Decoders/VideoDecoderVT.hpp
+++ b/Engine/Plugins/Experimental/AVCodecs/VTCodecs/Source/VTCodecs/Public/Video/Decoders/VideoDecoderVT.hpp
@@ -105,9 +105,13 @@ FAVResult TVideoDecoderVT<TResource>::ApplyConfig()
 
 			if (!IsInitialized())
 			{
-                // Set source image buffer attributes. These attributes will be present on
+                if (!PendingConfig.VideoFormat)
+				{
+					return FAVResult(EAVResult::PendingInput);
+				}
+				// Set source image buffer attributes. These attributes will be present on
                 // buffers retrieved from the decoder's pixel buffer pool.
-                CFMutableDictionaryRef SourceAttributes = CFDictionaryCreateMutable(kCFAllocatorDefault, 1, &kCFTypeDictionaryKeyCallBacks, &kCFTypeDictionaryValueCallBacks);
+                CFMutableDictionaryRef SourceAttributes = CFDictionaryCreateMutable(kCFAllocatorDefault, 0, &kCFTypeDictionaryKeyCallBacks, &kCFTypeDictionaryValueCallBacks);
                 CFDictionarySetValue(SourceAttributes, kCVPixelBufferOpenGLCompatibilityKey, kCFBooleanTrue);
                 CFDictionaryRef IOSurfaceValue = CFDictionaryCreate(kCFAllocatorDefault, nullptr, nullptr, 0, &kCFTypeDictionaryKeyCallBacks, &kCFTypeDictionaryValueCallBacks);
                 CFDictionarySetValue(SourceAttributes, kCVPixelBufferIOSurfacePropertiesKey, IOSurfaceValue);
@@ -172,19 +176,18 @@ FAVResult TVideoDecoderVT<TResource>::SendPacket(FVideoPacket const& Packet)
             return FAVResult(EAVResult::Error, TEXT("Unsupported codec"), TEXT("VT"));
         }
 
-        if (InputFormat && !CMFormatDescriptionEqual(InputFormat, this->AppliedConfig.VideoFormat))
+        if (InputFormat && (!this->AppliedConfig.VideoFormat || !CMFormatDescriptionEqual(InputFormat, this->AppliedConfig.VideoFormat)))
         {
             PendingConfig.SetVideoFormat(InputFormat);
         }
         
         FAVResult AVResult = ApplyConfig();
+		CONDITIONAL_RELEASE(InputFormat);
         if (AVResult.IsNotSuccess())
         {
             return AVResult;
         }
 
-        CONDITIONAL_RELEASE(InputFormat);
-
         if(!this->AppliedConfig.VideoFormat)
         {
             return FAVResult(EAVResult::WarningInvalidState, TEXT("Missing video format. Frame with sps/pps required."), TEXT("VT"));
