[AVCodecs, PS2] Fix: VideoToolbox only decoding a few frames

- Changed TVideoDecoderVT to be FVideoDecoderVT because it only every supported FVideoResourceMetal so no point in it being templated.
- Added a CopyFrommethod to FVideoResourceMetal to copy the decoded pixel buffer straight to an existing texture.
- Replaced usages of CVPixelBufferRef in FVideoResourceMetal to be MTL::Texture*
- Updated VideoEncoderVT to get a CVPixelBufferRef from the MTL::Texture*. This is the same logic that was in VideoResourceRHI.cpp, now just moved to VideoEncoderVT.hpp
- Updated TEpicRtcVideoDecoder to only ever use the same resolvable resource (no more flip flopping or pooling). This means there's only 1 texture per decoder that is updated using the relevant hardware API with the newly decoded data

diff --git a/Engine/Plugins/Experimental/AVCodecs/AVCodecsCore/Source/AVCodecsCore/Private/Video/Resources/Metal/VideoResourceMetal.cpp b/Engine/Plugins/Experimental/AVCodecs/AVCodecsCore/Source/AVCodecsCore/Private/Video/Resources/Metal/VideoResourceMetal.cpp
--- a/Engine/Plugins/Experimental/AVCodecs/AVCodecsCore/Source/AVCodecsCore/Private/Video/Resources/Metal/VideoResourceMetal.cpp
+++ b/Engine/Plugins/Experimental/AVCodecs/AVCodecsCore/Source/AVCodecsCore/Private/Video/Resources/Metal/VideoResourceMetal.cpp
@@ -4,19 +4,25 @@
 
 #include "Video/Resources/Metal/VideoResourceMetal.h"
 
+THIRD_PARTY_INCLUDES_START
+#include <VideoToolbox/VideoToolbox.h>
+#include <CoreMedia/CMSync.h>
+THIRD_PARTY_INCLUDES_END
+
 REGISTER_TYPEID(FVideoContextMetal);
 REGISTER_TYPEID(FVideoResourceMetal);
 
-static TAVResult<EVideoFormat> ConvertFormat(OSType Format)
+static TAVResult<EVideoFormat> ConvertFormat(MTL::PixelFormat Format)
 {
 	switch (Format)
 	{
-        case kCVPixelFormatType_32BGRA:
-		return EVideoFormat::BGRA;
-	case kCVPixelFormatType_ARGB2101010LEPacked:
-		return EVideoFormat::ABGR10;
-	default:
-		return FAVResult(EAVResult::ErrorUnsupported, FString::Printf(TEXT("MTL::PixelFormat format %d is not supported"), Format), TEXT("Metal"));
+		case MTL::PixelFormatBGRA8Unorm:
+		case MTL::PixelFormatBGRA8Unorm_sRGB:
+			return EVideoFormat::BGRA;
+		case MTL::PixelFormatBGR10A2Unorm:
+			return EVideoFormat::ABGR10;
+		default:
+			return FAVResult(EAVResult::ErrorUnsupported, FString::Printf(TEXT("MTL::PixelFormat format %d is not supported"), Format), TEXT("Metal"));
 	}
 }
 
@@ -25,32 +31,37 @@ FVideoContextMetal::FVideoContextMetal(MTL::Device* Device)
 {
 }
 
-FVideoDescriptor FVideoResourceMetal::GetDescriptorFrom(TSharedRef<FAVDevice> const& Device, CVPixelBufferRef Raw)
+FVideoDescriptor FVideoResourceMetal::GetDescriptorFrom(TSharedRef<FAVDevice> const& Device, MTL::Texture* Raw)
 {
-    uint32_t Width = CVPixelBufferGetWidth(Raw);
-    uint32_t Height = CVPixelBufferGetHeight(Raw);
-    TAVResult<EVideoFormat> ConvertedFormat = ConvertFormat(CVPixelBufferGetPixelFormatType(Raw));
+    uint32_t Width = Raw->width();
+    uint32_t Height = Raw->height();
+    TAVResult<EVideoFormat> ConvertedFormat = ConvertFormat(Raw->pixelFormat());
     
 	return FVideoDescriptor(ConvertedFormat, Width, Height);
 }
 
-FVideoResourceMetal::FVideoResourceMetal(TSharedRef<FAVDevice> const& Device, CVPixelBufferRef Raw, FAVLayout const& Layout)
+FVideoResourceMetal::FVideoResourceMetal(TSharedRef<FAVDevice> const& Device, MTL::Texture* Raw, FAVLayout const& Layout)
 	: TVideoResource(Device, Layout, GetDescriptorFrom(Device, Raw))
 	, Raw(Raw)
 {
-    if (Raw)
-    {
-        CFRetain(Raw);
-    }
 }
 
 FVideoResourceMetal::~FVideoResourceMetal()
 {
-    if (Raw)
-    {
-        CFRelease(Raw);
-        Raw = nullptr;
-    }
+}
+
+FAVResult FVideoResourceMetal::CopyFrom(CVPixelBufferRef Other)
+{	
+	if (const CVReturn Result = CVPixelBufferLockBaseAddress(Other, 0); Result != kCVReturnSuccess)
+	{
+		return FAVResult(EAVResult::Error, TEXT("Failed to lock input pixel buffer!"), TEXT("Metal"));
+	}
+	
+	Raw->replaceRegion(MTL::Region(0, 0, GetDescriptor().Width, GetDescriptor().Height), 0, 0, CVPixelBufferGetBaseAddress(Other), CVPixelBufferGetBytesPerRow(Other), 0);
+	
+	CVPixelBufferUnlockBaseAddress(Other, 0);
+	
+	return EAVResult::Success;
 }
 
 FAVResult FVideoResourceMetal::Validate() const
diff --git a/Engine/Plugins/Experimental/AVCodecs/AVCodecsCore/Source/AVCodecsCore/Public/Video/Resources/Metal/VideoResourceMetal.h b/Engine/Plugins/Experimental/AVCodecs/AVCodecsCore/Source/AVCodecsCore/Public/Video/Resources/Metal/VideoResourceMetal.h
--- a/Engine/Plugins/Experimental/AVCodecs/AVCodecsCore/Source/AVCodecsCore/Public/Video/Resources/Metal/VideoResourceMetal.h
+++ b/Engine/Plugins/Experimental/AVCodecs/AVCodecsCore/Source/AVCodecsCore/Public/Video/Resources/Metal/VideoResourceMetal.h
@@ -31,15 +31,17 @@ public:
 class AVCODECSCORE_API FVideoResourceMetal : public TVideoResource<FVideoContextMetal>
 {
 private:
-    CVPixelBufferRef Raw;
+    MTL::Texture* Raw;
 
 public:
-	static FVideoDescriptor GetDescriptorFrom(TSharedRef<FAVDevice> const& Device, CVPixelBufferRef Raw);
+	static FVideoDescriptor GetDescriptorFrom(TSharedRef<FAVDevice> const& Device, MTL::Texture* Raw);
 
-	FORCEINLINE CVPixelBufferRef GetRaw() const { return Raw; }
+	FORCEINLINE MTL::Texture* GetRaw() const { return Raw; }
 
-	FVideoResourceMetal(TSharedRef<FAVDevice> const& Device, CVPixelBufferRef Raw, FAVLayout const& Layout);
+	FVideoResourceMetal(TSharedRef<FAVDevice> const& Device, MTL::Texture* Raw, FAVLayout const& Layout);
     virtual ~FVideoResourceMetal() override;
+	
+	FAVResult CopyFrom(CVPixelBufferRef Other);
     
 	virtual FAVResult Validate() const override;
 };
diff --git a/Engine/Plugins/Experimental/AVCodecs/AVCodecsCore/Source/AVCodecsCoreRHI/Private/Video/Resources/VideoResourceRHI.cpp b/Engine/Plugins/Experimental/AVCodecs/AVCodecsCore/Source/AVCodecsCoreRHI/Private/Video/Resources/VideoResourceRHI.cpp
--- a/Engine/Plugins/Experimental/AVCodecs/AVCodecsCore/Source/AVCodecsCoreRHI/Private/Video/Resources/VideoResourceRHI.cpp
+++ b/Engine/Plugins/Experimental/AVCodecs/AVCodecsCore/Source/AVCodecsCoreRHI/Private/Video/Resources/VideoResourceRHI.cpp
@@ -587,60 +587,9 @@ DLLEXPORT FAVResult FAVExtension::TransformResource(TSharedPtr<FVideoResourceMet
 			FTextureRHIRef InResourceTexture = InResource->GetRaw().Texture;
             if(bool(InResourceTexture->GetDesc().Flags & TexCreate_CPUReadback))
             {
-                /*
-                 * NOTE (william.belcher): If the texture passed in has was created with 'ETextureCreateFlags::CPUReadback', we can simply copy the raw bytes
-                 * out of the texture. This is the preferred way as it prevents an ~50ms block on the encoding thread
-                 */
-				const FVideoDescriptor& Descriptor = InResource->GetDescriptor();
-            
-            	CFMutableDictionaryRef SourceAttributes = CFDictionaryCreateMutable(kCFAllocatorDefault, 1, &kCFTypeDictionaryKeyCallBacks, &kCFTypeDictionaryValueCallBacks);
-            	CFDictionarySetValue(SourceAttributes, kCVPixelBufferOpenGLCompatibilityKey, kCFBooleanTrue);
-	
-            	CFDictionaryRef IOSurfaceValue = CFDictionaryCreate(kCFAllocatorDefault, nullptr, nullptr, 0, &kCFTypeDictionaryKeyCallBacks, &kCFTypeDictionaryValueCallBacks);
-            	CFDictionarySetValue(SourceAttributes, kCVPixelBufferIOSurfacePropertiesKey, IOSurfaceValue);
-	
-            	int64 PixelType = 0;
-            	switch(Descriptor.Format)
-            	{
-            	    case EVideoFormat::BGRA:
-            	        PixelType = kCVPixelFormatType_32BGRA;
-            	        break;
-            	    case EVideoFormat::ABGR10:
-            	        PixelType = kCVPixelFormatType_ARGB2101010LEPacked;
-            	        break;
-            	    default:
-            	        checkNoEntry();
-            	}
-	
-            	CFNumberRef PixelFormat = CFNumberCreate(nullptr, kCFNumberLongType, &PixelType);
-            	CFDictionarySetValue(SourceAttributes, kCVPixelBufferPixelFormatTypeKey, PixelFormat);
-	
-            	CFSafeRelease(IOSurfaceValue);
-            	CFSafeRelease(PixelFormat);
-	
-            	CVPixelBufferRef PixelBuffer;
-            	CVReturn Result = CVPixelBufferCreate(kCFAllocatorDefault, Descriptor.Width, Descriptor.Height, PixelType, SourceAttributes, &PixelBuffer);
-            	if (Result != kCVReturnSuccess)
-            	{
-            	    return FAVResult(EAVResult::Error, TEXT("Failed to create CVPixelBufferRef"), TEXT("RHI"), Result);
-            	}
-            	CFSafeRelease(SourceAttributes);
-
-                Result = CVPixelBufferLockBaseAddress(PixelBuffer, 0);
-                if (Result != kCVReturnSuccess)
-                {
-                    return FAVResult(EAVResult::Error, TEXT("Failed to lock base address"), TEXT("RHI"), Result);
-                }
-
-                // NOTE (belchy06): GetBytes assumes the raw texture has been created with with TexCreate_CPUReadback
-                static_cast<MTL::Texture*>(InResource->GetRaw().Texture->GetNativeResource())->getBytes(reinterpret_cast<uint8*>(CVPixelBufferGetBaseAddressOfPlane(PixelBuffer, 0)), CVPixelBufferGetBytesPerRow(PixelBuffer), MTL::Region(0, 0, Descriptor.Width, Descriptor.Height), 0);
-                CVPixelBufferUnlockBaseAddress(PixelBuffer, 0);
-                
                 OutResource = MakeShared<FVideoResourceMetal>(InResource->GetDevice(),
-                                                              PixelBuffer,
+															  static_cast<MTL::Texture*>(InResource->GetRaw().Texture->GetNativeResource()),
                                                               InResource->GetLayout());
-
-				CVPixelBufferRelease(PixelBuffer);
 			
 			    return OutResource->Validate();
             }
diff --git a/Engine/Plugins/Experimental/AVCodecs/VTCodecs/Source/VTCodecs/Private/VTCodecsModule.cpp b/Engine/Plugins/Experimental/AVCodecs/VTCodecs/Source/VTCodecs/Private/VTCodecsModule.cpp
--- a/Engine/Plugins/Experimental/AVCodecs/VTCodecs/Source/VTCodecs/Private/VTCodecsModule.cpp
+++ b/Engine/Plugins/Experimental/AVCodecs/VTCodecs/Source/VTCodecs/Private/VTCodecsModule.cpp
@@ -88,7 +88,7 @@ public:
 
         if(VTIsHardwareDecodeSupported(kCMVideoCodecType_H264))
         {
-            FVideoDecoder::RegisterPermutationsOf<TVideoDecoderVT<FVideoResourceMetal>>
+            FVideoDecoder::RegisterPermutationsOf<FVideoDecoderVT>
                 ::With<FVideoResourceMetal>
                 ::And<FVideoDecoderConfigVT, FVideoDecoderConfigH264>(
                     [](TSharedRef<FAVDevice> const& NewDevice, TSharedRef<FAVInstance> const& NewInstance)
@@ -99,7 +99,7 @@ public:
 
         if(VTIsHardwareDecodeSupported(kCMVideoCodecType_HEVC))
         {
-            FVideoDecoder::RegisterPermutationsOf<TVideoDecoderVT<FVideoResourceMetal>>
+            FVideoDecoder::RegisterPermutationsOf<FVideoDecoderVT>
                 ::With<FVideoResourceMetal>
                 ::And<FVideoDecoderConfigVT, FVideoDecoderConfigH265>(
                     [](TSharedRef<FAVDevice> const& NewDevice, TSharedRef<FAVInstance> const& NewInstance)
@@ -110,7 +110,7 @@ public:
 
         if(VTIsHardwareDecodeSupported(kCMVideoCodecType_VP9))
         {
-            FVideoDecoder::RegisterPermutationsOf<TVideoDecoderVT<FVideoResourceMetal>>
+            FVideoDecoder::RegisterPermutationsOf<FVideoDecoderVT>
                 ::With<FVideoResourceMetal>
                 ::And<FVideoDecoderConfigVT, FVideoDecoderConfigVP9>(
                     [](TSharedRef<FAVDevice> const& NewDevice, TSharedRef<FAVInstance> const& NewInstance)
diff --git a/Engine/Plugins/Experimental/AVCodecs/VTCodecs/Source/VTCodecs/Private/Video/Decoders/VideoDecoderVT.cpp b/Engine/Plugins/Experimental/AVCodecs/VTCodecs/Source/VTCodecs/Private/Video/Decoders/VideoDecoderVT.cpp
--- a/Engine/Plugins/Experimental/AVCodecs/VTCodecs/Source/VTCodecs/Private/Video/Decoders/VideoDecoderVT.cpp
+++ b/Engine/Plugins/Experimental/AVCodecs/VTCodecs/Source/VTCodecs/Private/Video/Decoders/VideoDecoderVT.cpp
@@ -3,10 +3,290 @@
 #include "Video/Decoders/VideoDecoderVT.h"
 #include "Video/Resources/Metal/VideoResourceMetal.h"
 
+#include "Video/Util/NaluRewriter.h"
+
+#include "DynamicRHI.h"
+
+THIRD_PARTY_INCLUDES_START
+#include <VideoToolbox/VideoToolbox.h>
+#include <CoreMedia/CMSync.h>
+THIRD_PARTY_INCLUDES_END
+
+#define CONDITIONAL_RELEASE(x)          \
+	if (x)                              \
+	{                                   \
+		CFRelease(x);                   \
+		x = nullptr;                    \
+	}
+
+
 namespace Internal
 {
     void VTDecompressionOutputCallback(void* Decoder, void* Params, OSStatus Status, VTDecodeInfoFlags InfoFlags, CVImageBufferRef ImageBuffer, CMTime Timestamp, CMTime Duration)
     {
-        static_cast<TVideoDecoderVT<FVideoResourceMetal>*>(Decoder)->HandleFrame(Params, Status, InfoFlags, ImageBuffer, Timestamp, Duration);
+        static_cast<FVideoDecoderVT*>(Decoder)->HandleFrame(Params, Status, InfoFlags, ImageBuffer, Timestamp, Duration);
     }
 }
+
+FVideoDecoderVT::~FVideoDecoderVT()
+{
+	Close();
+}
+
+bool FVideoDecoderVT::IsOpen() const
+{
+	return bIsOpen;
+}
+
+FAVResult FVideoDecoderVT::Open(TSharedRef<FAVDevice> const& NewDevice, TSharedRef<FAVInstance> const& NewInstance)
+{
+	Close();
+
+	TVideoDecoder<FVideoResourceMetal, FVideoDecoderConfigVT>::Open(NewDevice, NewInstance);
+
+	MemoryPool = CMMemoryPoolCreate(nullptr);
+
+	FrameCount = 0;
+
+	bIsOpen = true;
+
+	return EAVResult::Success;
+}
+
+void FVideoDecoderVT::Close()
+{
+	DestroyDecompressionSession();
+
+	if(MemoryPool)
+	{
+		CMMemoryPoolInvalidate(MemoryPool);
+		CFRelease(MemoryPool);
+	}
+
+	bIsOpen = false;
+}
+
+void FVideoDecoderVT::DestroyDecompressionSession()
+{
+	if (Decoder) 
+	{
+		VTDecompressionSessionInvalidate(Decoder);	
+		CFRelease(Decoder);
+		Decoder = nullptr;
+	}
+}
+
+void FVideoDecoderVT::ConfigureDecompressionSession()
+{
+	VTSessionSetProperty(Decoder, kVTDecompressionPropertyKey_RealTime, kCFBooleanTrue);
+}
+
+bool FVideoDecoderVT::IsInitialized() const
+{
+	return Decoder != nullptr;
+}
+
+FAVResult FVideoDecoderVT::ApplyConfig()
+{
+	if (IsOpen())
+	{
+		FVideoDecoderConfigVT const& PendingConfig = this->GetPendingConfig();
+		if (this->AppliedConfig != PendingConfig)
+		{
+			if (IsInitialized())
+			{
+				// VideoToolbox decoder doesn't support reconfiguration. If any aspect of the config changes,
+				// the entire session must be re-created
+				if (Decoder)
+				{
+					DestroyDecompressionSession();
+					FAVResult::Log(EAVResult::Success, TEXT("Re-initializing decoding session"), TEXT("VT"));
+				}
+			}
+
+			if (!IsInitialized())
+			{
+				if (!PendingConfig.VideoFormat)
+				{
+					return FAVResult(EAVResult::PendingInput);
+				}
+
+				// Set source image buffer attributes. These attributes will be present on
+				// buffers retrieved from the decoder's pixel buffer pool.
+				CFMutableDictionaryRef SourceAttributes = CFDictionaryCreateMutable(kCFAllocatorDefault, 3, &kCFTypeDictionaryKeyCallBacks, &kCFTypeDictionaryValueCallBacks);
+				CFDictionarySetValue(SourceAttributes, kCVPixelBufferOpenGLCompatibilityKey, kCFBooleanTrue);
+				CFDictionaryRef IOSurfaceValue = CFDictionaryCreate(kCFAllocatorDefault, nullptr, nullptr, 0, &kCFTypeDictionaryKeyCallBacks, &kCFTypeDictionaryValueCallBacks);
+				CFDictionarySetValue(SourceAttributes, kCVPixelBufferIOSurfacePropertiesKey, IOSurfaceValue);
+				// TODO (belchy06): This should support more than ARGB8 pixel format
+				int64 PixelType = kCVPixelFormatType_32BGRA;
+				CFNumberRef PixelFormat = CFNumberCreate(nullptr, kCFNumberLongType, &PixelType);
+				CFDictionarySetValue(SourceAttributes, kCVPixelBufferPixelFormatTypeKey, PixelFormat);
+
+				CONDITIONAL_RELEASE(IOSurfaceValue);
+				CONDITIONAL_RELEASE(PixelFormat);
+
+				VTDecompressionOutputCallbackRecord Record = 
+				{
+					Internal::VTDecompressionOutputCallback, 
+					this
+				};
+
+				OSStatus Result = VTDecompressionSessionCreate(kCFAllocatorDefault, PendingConfig.VideoFormat, nullptr, SourceAttributes, &Record, &Decoder);
+
+				CONDITIONAL_RELEASE(SourceAttributes);
+
+				if(Result != 0)
+				{
+					DestroyDecompressionSession();
+					return FAVResult(EAVResult::ErrorCreating, TEXT("Failed to create VTDecompressionSession"), TEXT("VT"), Result);
+				}
+
+				ConfigureDecompressionSession();
+			}
+		}
+
+		return TVideoDecoder<FVideoResourceMetal, FVideoDecoderConfigVT>::ApplyConfig();
+	}
+
+	return FAVResult(EAVResult::ErrorInvalidState, TEXT("Decoder not open"), TEXT("VT"));
+}
+
+FAVResult FVideoDecoderVT::SendPacket(FVideoPacket const& Packet)
+{
+	if (IsOpen())
+	{
+		// We've received a call to decode a frame, we can now parse the information from
+		// the bitstream, configure our config and initialize the session
+		FVideoDecoderConfigVT& PendingConfig = this->EditPendingConfig();
+
+		CMVideoFormatDescriptionRef InputFormat = nullptr;
+		if(PendingConfig.Codec == kCMVideoCodecType_H264)
+		{
+			InputFormat = NaluRewriter::CreateH264VideoFormatDescription(Packet.DataPtr.Get(), Packet.DataSize);
+		}
+		else if(PendingConfig.Codec == kCMVideoCodecType_HEVC)
+		{
+			InputFormat = NaluRewriter::CreateH265VideoFormatDescription(Packet.DataPtr.Get(), Packet.DataSize);
+		}
+		else if (PendingConfig.Codec == kCMVideoCodecType_VP9)
+		{
+			InputFormat = NaluRewriter::CreateVP9VideoFormatDescription(Packet.DataPtr.Get(), Packet.DataSize);
+		}
+		else
+		{
+			return FAVResult(EAVResult::Error, TEXT("Unsupported codec"), TEXT("VT"));
+		}
+
+		if (InputFormat && (!this->AppliedConfig.VideoFormat || !CMFormatDescriptionEqual(InputFormat, this->AppliedConfig.VideoFormat)))
+		{
+			PendingConfig.SetVideoFormat(InputFormat);
+		}
+		
+		FAVResult AVResult = ApplyConfig();
+		CONDITIONAL_RELEASE(InputFormat);
+		if (AVResult.IsNotSuccess())
+		{
+			return AVResult;
+		}
+
+		if(!this->AppliedConfig.VideoFormat)
+		{
+			return FAVResult(EAVResult::WarningInvalidState, TEXT("Missing video format. Frame with sps/pps required."), TEXT("VT"));
+		}
+
+		CMSampleBufferRef SampleBuffer = nullptr;
+		if(this->AppliedConfig.Codec == kCMVideoCodecType_H264)
+		{
+			if(!NaluRewriter::H264AnnexBBufferToCMSampleBuffer(Packet.DataPtr.Get(), Packet.DataSize, this->AppliedConfig.VideoFormat, &SampleBuffer, MemoryPool))
+			{
+				return FAVResult(EAVResult::Error, TEXT("Failed to get SampleBuffer"), TEXT("VT"));
+			}
+		}
+		else if(this->AppliedConfig.Codec == kCMVideoCodecType_HEVC)
+		{
+			if(!NaluRewriter::H265AnnexBBufferToCMSampleBuffer(Packet.DataPtr.Get(), Packet.DataSize, this->AppliedConfig.VideoFormat, &SampleBuffer, MemoryPool))
+			{
+				return FAVResult(EAVResult::Error, TEXT("Failed to get SampleBuffer"), TEXT("VT"));
+			}
+		}
+		else if (this->AppliedConfig.Codec == kCMVideoCodecType_VP9)
+		{
+			if(!NaluRewriter::VP9BufferToCMSampleBuffer(Packet.DataPtr.Get(), Packet.DataSize, this->AppliedConfig.VideoFormat, &SampleBuffer, MemoryPool))
+			{
+				return FAVResult(EAVResult::Error, TEXT("Failed to get SampleBuffer"), TEXT("VT"));
+			}
+		}
+		else
+		{
+			return FAVResult(EAVResult::Error, TEXT("Unsupported codec"), TEXT("VT"));
+		}
+		
+		if(SampleBuffer == nullptr)
+		{
+			return FAVResult(EAVResult::Error, TEXT("SampleBuffer is nullptr"), TEXT("VT"));
+		}
+
+		OSStatus Result = VTDecompressionSessionDecodeFrame(Decoder, SampleBuffer, 0, nullptr, nullptr);
+
+		CFRelease(SampleBuffer);
+
+		if(Result != 0)
+		{
+			return FAVResult(EAVResult::Error, TEXT("Failed to decode frame"), TEXT("VT"), Result);
+		}
+
+		return EAVResult::Success;
+	}
+
+	return FAVResult(EAVResult::ErrorInvalidState, TEXT("Decoder not open"), TEXT("VT"));
+}
+
+FAVResult FVideoDecoderVT::ReceiveFrame(TResolvableVideoResource<FVideoResourceMetal>& InOutResource)
+{
+	if(IsOpen())
+	{
+		if(Frames.Peek())
+		{
+			TSharedPtr<FFrame> Frame = *Frames.Peek();
+			Frames.Pop();
+
+			size_t Width = CVPixelBufferGetWidth(Frame->ImageBuffer);
+			size_t Height = CVPixelBufferGetHeight(Frame->ImageBuffer);		
+
+			if (!InOutResource.Resolve(this->GetDevice(), FVideoDescriptor(EVideoFormat::BGRA, Width, Height)))
+			{
+				return FAVResult(EAVResult::ErrorResolving, TEXT("Failed to resolve frame resource"), TEXT("VT"));
+			}
+
+			return InOutResource->CopyFrom(Frame->ImageBuffer);
+		}
+
+		return EAVResult::PendingInput;
+	}
+
+	return FAVResult(EAVResult::ErrorInvalidState, TEXT("Decoder not open"), TEXT("VT"));
+}
+
+FAVResult FVideoDecoderVT::HandleFrame(void* Params, OSStatus Status, VTDecodeInfoFlags InfoFlags, CVImageBufferRef ImageBuffer, CMTime Timestamp, CMTime Duration)
+{
+	if(IsOpen())
+	{
+		if(Status != 0)
+		{
+			return FAVResult(EAVResult::Error, TEXT("Failed to decode"), TEXT("VT"), Status);
+		}
+
+		if(!ImageBuffer)
+		{
+			return FAVResult(EAVResult::Error, TEXT("No output image buffer"), TEXT("VT"), Status);
+		}
+		
+		// The destructor for FFrame releases the ImageBuffer so we need to make sure it's not release until after it's been used
+		TSharedPtr<FFrame> Frame = MakeShareable(new FFrame(ImageBuffer, Timestamp, Duration));
+		Frames.Enqueue(Frame);
+
+		return EAVResult::Success;
+	}
+	
+	return FAVResult(EAVResult::ErrorInvalidState, TEXT("Decoder not open"), TEXT("VT"));
+}
diff --git a/Engine/Plugins/Experimental/AVCodecs/VTCodecs/Source/VTCodecs/Public/Video/Decoders/VideoDecoderVT.h b/Engine/Plugins/Experimental/AVCodecs/VTCodecs/Source/VTCodecs/Public/Video/Decoders/VideoDecoderVT.h
--- a/Engine/Plugins/Experimental/AVCodecs/VTCodecs/Source/VTCodecs/Public/Video/Decoders/VideoDecoderVT.h
+++ b/Engine/Plugins/Experimental/AVCodecs/VTCodecs/Source/VTCodecs/Public/Video/Decoders/VideoDecoderVT.h
@@ -16,8 +16,7 @@ THIRD_PARTY_INCLUDES_START
 #include <CoreMedia/CoreMedia.h>
 THIRD_PARTY_INCLUDES_END
 
-template <typename TResource>
-class TVideoDecoderVT : public TVideoDecoder<TResource, FVideoDecoderConfigVT>
+class FVideoDecoderVT : public TVideoDecoder<FVideoResourceMetal, FVideoDecoderConfigVT>
 {
 private:
     VTDecompressionSessionRef Decoder = nullptr;
@@ -70,8 +69,8 @@ private:
     TQueue<TSharedPtr<FFrame>> Frames;
 
 public:	
-	TVideoDecoderVT() = default;
-	virtual ~TVideoDecoderVT() override;
+	FVideoDecoderVT() = default;
+	virtual ~FVideoDecoderVT() override;
 
 	virtual bool IsOpen() const override;
 	virtual FAVResult Open(TSharedRef<FAVDevice> const& NewDevice, TSharedRef<FAVInstance> const& NewInstance) override;
@@ -83,7 +82,7 @@ public:
 
 	virtual FAVResult ApplyConfig() override;
 	
-	virtual FAVResult ReceiveFrame(TResolvableVideoResource<TResource>& InOutResource) override;
+	virtual FAVResult ReceiveFrame(TResolvableVideoResource<FVideoResourceMetal>& InOutResource) override;
     FAVResult HandleFrame(void* Params, OSStatus Status, VTDecodeInfoFlags InfoFlags, CVImageBufferRef ImageBuffer, CMTime Timestamp, CMTime Duration);
 
 	virtual FAVResult SendPacket(FVideoPacket const& Packet) override;
@@ -93,6 +92,3 @@ namespace Internal
 {
     void VTDecompressionOutputCallback(void* Decoder, void* Params, OSStatus Status, VTDecodeInfoFlags InfoFlags, CVImageBufferRef ImageBuffer, CMTime Timestamp, CMTime Duration);
 }
-
-#include "VideoDecoderVT.hpp"
-
diff --git a/Engine/Plugins/Experimental/AVCodecs/VTCodecs/Source/VTCodecs/Public/Video/Decoders/VideoDecoderVT.hpp b/Engine/Plugins/Experimental/AVCodecs/VTCodecs/Source/VTCodecs/Public/Video/Decoders/VideoDecoderVT.hpp
deleted file mode 100644
diff --git a/Engine/Plugins/Experimental/AVCodecs/VTCodecs/Source/VTCodecs/Public/Video/Encoders/VideoEncoderVT.hpp b/Engine/Plugins/Experimental/AVCodecs/VTCodecs/Source/VTCodecs/Public/Video/Encoders/VideoEncoderVT.hpp
--- a/Engine/Plugins/Experimental/AVCodecs/VTCodecs/Source/VTCodecs/Public/Video/Encoders/VideoEncoderVT.hpp
+++ b/Engine/Plugins/Experimental/AVCodecs/VTCodecs/Source/VTCodecs/Public/Video/Encoders/VideoEncoderVT.hpp
@@ -226,10 +226,60 @@ FAVResult TVideoEncoderVT<TResource>::SendFrame(TSharedPtr<FVideoResourceMetal>
 
             TUniquePtr<EncodeParams> Params = MakeUnique<EncodeParams>();
             Params.Reset(new EncodeParams(PendingConfig.Codec, PresentationTime));
-
-            OSStatus Status = VTCompressionSessionEncodeFrame(Encoder, Resource->GetRaw(), PresentationTime, kCMTimeInvalid, FrameProperties, (void *)Params.Release(), NULL);
+			
+			/*
+			 * NOTE (william.belcher): If the texture passed in has was created with 'ETextureCreateFlags::CPUReadback', we can simply copy the raw bytes
+			 * out of the texture. This is the preferred way as it prevents an ~50ms block on the encoding thread
+			 */
+		
+			CFMutableDictionaryRef SourceAttributes = CFDictionaryCreateMutable(kCFAllocatorDefault, 1, &kCFTypeDictionaryKeyCallBacks, &kCFTypeDictionaryValueCallBacks);
+			CFDictionarySetValue(SourceAttributes, kCVPixelBufferOpenGLCompatibilityKey, kCFBooleanTrue);
+
+			CFDictionaryRef IOSurfaceValue = CFDictionaryCreate(kCFAllocatorDefault, nullptr, nullptr, 0, &kCFTypeDictionaryKeyCallBacks, &kCFTypeDictionaryValueCallBacks);
+			CFDictionarySetValue(SourceAttributes, kCVPixelBufferIOSurfacePropertiesKey, IOSurfaceValue);
+
+			int64 PixelType = 0;
+			switch(Resource->GetDescriptor().Format)
+			{
+				case EVideoFormat::BGRA:
+					PixelType = kCVPixelFormatType_32BGRA;
+					break;
+				case EVideoFormat::ABGR10:
+					PixelType = kCVPixelFormatType_ARGB2101010LEPacked;
+					break;
+				default:
+					checkNoEntry();
+			}
+
+			CFNumberRef PixelFormat = CFNumberCreate(nullptr, kCFNumberLongType, &PixelType);
+			CFDictionarySetValue(SourceAttributes, kCVPixelBufferPixelFormatTypeKey, PixelFormat);
+
+			CONDITIONAL_RELEASE(IOSurfaceValue);
+			CONDITIONAL_RELEASE(PixelFormat);
+
+			CVPixelBufferRef PixelBuffer;
+			CVReturn Result = CVPixelBufferCreate(kCFAllocatorDefault, Resource->GetDescriptor().Width, Resource->GetDescriptor().Height, PixelType, SourceAttributes, &PixelBuffer);
+			if (Result != kCVReturnSuccess)
+			{
+				return FAVResult(EAVResult::Error, TEXT("Failed to create CVPixelBufferRef"), TEXT("RHI"), Result);
+			}
+			CONDITIONAL_RELEASE(SourceAttributes);
+
+			Result = CVPixelBufferLockBaseAddress(PixelBuffer, 0);
+			if (Result != kCVReturnSuccess)
+			{
+				return FAVResult(EAVResult::Error, TEXT("Failed to lock base address"), TEXT("RHI"), Result);
+			}
+
+			// NOTE (belchy06): GetBytes assumes the raw texture has been created with with TexCreate_CPUReadback
+			Resource->GetRaw()->getBytes(reinterpret_cast<uint8*>(CVPixelBufferGetBaseAddressOfPlane(PixelBuffer, 0)), CVPixelBufferGetBytesPerRow(PixelBuffer), MTL::Region(0, 0, Resource->GetDescriptor().Width, Resource->GetDescriptor().Height), 0);
+			CVPixelBufferUnlockBaseAddress(PixelBuffer, 0);
+			
+            OSStatus Status = VTCompressionSessionEncodeFrame(Encoder, PixelBuffer, PresentationTime, kCMTimeInvalid, FrameProperties, (void *)Params.Release(), NULL);
 
             CONDITIONAL_RELEASE(FrameProperties);
+			
+			CONDITIONAL_RELEASE(PixelBuffer);
 
             if (Status != kCVReturnSuccess)
             {
diff --git a/Engine/Plugins/Media/PixelStreaming2/Source/PixelStreaming2Editor/Private/VideoProducerBackBufferComposited.cpp b/Engine/Plugins/Media/PixelStreaming2/Source/PixelStreaming2Editor/Private/VideoProducerBackBufferComposited.cpp
--- a/Engine/Plugins/Media/PixelStreaming2/Source/PixelStreaming2Editor/Private/VideoProducerBackBufferComposited.cpp
+++ b/Engine/Plugins/Media/PixelStreaming2/Source/PixelStreaming2Editor/Private/VideoProducerBackBufferComposited.cpp
@@ -31,7 +31,7 @@ namespace UE::EditorPixelStreaming2
 		{
 			TexCreateFlags |= ETextureCreateFlags::Shared;
 		}
-
+		
 		return TexCreateFlags;
 	}
 
@@ -255,11 +255,19 @@ namespace UE::EditorPixelStreaming2
 				// Add screen pass to convert whatever format the editor produces to BGRA8
 				AddDrawScreenPass(GraphBuilder, RDG_EVENT_NAME("VideoProducerBackBufferCompositedSwizzle"), View, OutputViewport, InputViewport, VertexShader, PixelShader, PixelShaderParameters);
 			}
-
+			
 			// Final pass to extract the composited frames underlying RHI resource for passing to the rest of the pixel streaming pipeline
 			// GraphBuilder.SetTextureAccessFinal(CompositedTexture, ERHIAccess::CopySrc);
+#if PLATFORM_MAC
+			// On Mac specifically, we need to add one more pass to render to a cpu readable texture for AVCodecs
+			FRDGTextureRef StagingTexture = GraphBuilder.CreateTexture(FRDGTextureDesc::Create2D(CompositedTexture->Desc.Extent, EPixelFormat::PF_B8G8R8A8, FClearValueBinding::None, ETextureCreateFlags::CPUReadback), TEXT("VideoProducerBackBufferComposited MAC Staging"));
+			
+			AddDrawTexturePass(GraphBuilder, GetGlobalShaderMap(GMaxRHIFeatureLevel), CompositedTexture, StagingTexture, FRDGDrawTextureInfo());
+
+			GraphBuilder.QueueTextureExtraction(StagingTexture, &OutTexture);
+#else
 			GraphBuilder.QueueTextureExtraction(CompositedTexture, &OutTexture);
-
+#endif
 			GraphBuilder.Execute();
 		}
 
@@ -276,4 +284,4 @@ namespace UE::EditorPixelStreaming2
 		return TEXT("the Editor");
 	}
 
-} // namespace UE::EditorPixelStreaming2
\ No newline at end of file
+} // namespace UE::EditorPixelStreaming2
diff --git a/Engine/Plugins/Media/PixelStreaming2/Source/PixelStreaming2RTC/Private/EpicRtcVideoDecoder.cpp b/Engine/Plugins/Media/PixelStreaming2/Source/PixelStreaming2RTC/Private/EpicRtcVideoDecoder.cpp
--- a/Engine/Plugins/Media/PixelStreaming2/Source/PixelStreaming2RTC/Private/EpicRtcVideoDecoder.cpp
+++ b/Engine/Plugins/Media/PixelStreaming2/Source/PixelStreaming2RTC/Private/EpicRtcVideoDecoder.cpp
@@ -113,77 +113,58 @@ namespace UE::PixelStreaming2
 				._buffer = nullptr
 			};
 
-			FAVResult DecodeResult;
-			if constexpr (std::is_same_v<TVideoResource, FVideoResourceRHI>)
+			FAVResult DecodeResult = Decoder->ReceiveFrame(DecoderResource);
+			if (DecodeResult.IsSuccess())
 			{
-				// TODO (Eden.Harris) RTCP-7927 Use VideoResources and FetchOrCreate rather than flip flopping.
-				// Note: GetOrCreate currently has a sync issue where old frames are displayed.
-				// By flip flopping VideoResourceIndex, old frames may be overwritten with a new frame texture.
-				// This results in low latency but overwritten frames and resolves old frames accidentally being shown out of order.
-				// FResolvableVideoResourceRHI& DecoderResource = VideoResources.GetOrCreate();
-				FResolvableVideoResourceRHI& DecoderResource = VideoResourcesRHI[VideoResourceIndex++ % 2];
-				DecodeResult = Decoder->ReceiveFrame(DecoderResource);
-
-				if (DecodeResult.IsSuccess())
+				
+				if constexpr (std::is_same_v<TVideoResource, FVideoResourceRHI>)
 				{
 					DecodedFrame._buffer = new FEpicRtcVideoBufferRHI(DecoderResource);
 				}
-			}
-			else if constexpr (std::is_same_v<TVideoResource, FVideoResourceCPU>)
-			{
-				// TODO (Eden.Harris) RTCP-7927 Use VideoResources and FetchOrCreate rather than flip flopping.
-				// Note: GetOrCreate currently has a sync issue where old frames are displayed.
-				// By flip flopping VideoResourceIndex, old frames may be overwritten with a new frame texture.
-				// This results in low latency but overwritten frames and resolves old frames accidentally being shown out of order.
-				// FResolvableVideoResourceCPU& DecoderResource = VideoResources.GetOrCreate();
-				FResolvableVideoResourceCPU& DecoderResource = VideoResourcesCPU[VideoResourceIndex++ % 2];
-				DecodeResult = Decoder->ReceiveFrame(DecoderResource);
-
-				if (DecodeResult.IsSuccess())
+				else if constexpr (std::is_same_v<TVideoResource, FVideoResourceCPU>)
 				{
 					const uint32 FrameWidth = DecoderResource->GetWidth();
 					const uint32 FrameHeight = DecoderResource->GetHeight();
 					
 					TSharedPtr<FPixelCaptureBufferI420> I420Buffer = MakeShared<FPixelCaptureBufferI420>(FrameWidth, FrameHeight);
-
+					
 					const uint32 DataSizeY = FrameWidth * FrameHeight;
 					const uint32 DataSizeUV = ((FrameWidth + 1) / 2) * ((FrameHeight + 1) / 2);
-
+					
 					CopyI420(
-						DecoderResource->GetRaw().Get(), FrameWidth,
-						DecoderResource->GetRaw().Get() + DataSizeY, (FrameWidth + 1) / 2,
-						DecoderResource->GetRaw().Get() + DataSizeY + DataSizeUV, (FrameWidth + 1) / 2,
-						I420Buffer->GetMutableDataY(), I420Buffer->GetStrideY(),
-						I420Buffer->GetMutableDataU(), I420Buffer->GetStrideUV(),
-						I420Buffer->GetMutableDataV(), I420Buffer->GetStrideUV(),
-						FrameWidth, FrameHeight);
-
+							 DecoderResource->GetRaw().Get(), FrameWidth,
+							 DecoderResource->GetRaw().Get() + DataSizeY, (FrameWidth + 1) / 2,
+							 DecoderResource->GetRaw().Get() + DataSizeY + DataSizeUV, (FrameWidth + 1) / 2,
+							 I420Buffer->GetMutableDataY(), I420Buffer->GetStrideY(),
+							 I420Buffer->GetMutableDataU(), I420Buffer->GetStrideUV(),
+							 I420Buffer->GetMutableDataV(), I420Buffer->GetStrideUV(),
+							 FrameWidth, FrameHeight);
+					
 					DecodedFrame._buffer = new FEpicRtcVideoBufferI420(I420Buffer);
 				}
-			}
-			else
-			{
-				UE_LOGFMT(LogPixelStreaming2RTC, Error, "VideoResource isn't a compatible type! Expected either a FVideoResourceRHI or FVideoResourceCPU. Received: {0}", PREPROCESSOR_TO_STRING(TVideoResource));
-				return EpicRtcMediaResult::Error;
-			}
-
-			if (DecodeResult.IsSuccess())
-			{
+				else
+				{
+					UE_LOGFMT(LogPixelStreaming2RTC, Error, "VideoResource isn't a compatible type! Expected either a FVideoResourceRHI or FVideoResourceCPU. Received: {0}", PREPROCESSOR_TO_STRING(TVideoResource));
+					return EpicRtcMediaResult::Error;
+				}	
+				
 				check(DecodedFrame._buffer->GetWidth() != 0 && DecodedFrame._buffer->GetHeight() != 0);
-
+				
 				CallbackDecoded->Decoded(DecodedFrame, FPlatformTime::ToMilliseconds64(FPlatformTime::Cycles64()) - TimestampDecodeStart, Frame._qp);
-
+				
 				return EpicRtcMediaResult::Ok;
 			}
+			else 
+			{
+				UE_LOG(LogPixelStreaming2RTC, Warning, TEXT("TEpicRtcVideoDecoder::Decode FAILED TO RECEIVE FRAME"));				
+				return EpicRtcMediaResult::Error;
+			}
 		}
 		else
 		{
-			UE_LOG(LogPixelStreaming2RTC, Warning, TEXT("FVideoDecoderHardware::Decode FAILED"));
+			UE_LOG(LogPixelStreaming2RTC, Warning, TEXT("TEpicRtcVideoDecoder::Decode DECODE FAILED"));
 			return EpicRtcMediaResult::OkRequestKeyframe;
 		}
-
-		UE_LOG(LogPixelStreaming2RTC, Error, TEXT("FVideoDecoderHardware::Decode ERROR"));
-		return EpicRtcMediaResult::Error;
 	}
 
 	template <std::derived_from<FVideoResource> TVideoResource>
@@ -259,4 +240,4 @@ namespace UE::PixelStreaming2
 	// Explicit specialisation
 	template class TEpicRtcVideoDecoder<FVideoResourceRHI>;
 	template class TEpicRtcVideoDecoder<FVideoResourceCPU>;
-} // namespace UE::PixelStreaming2
\ No newline at end of file
+} // namespace UE::PixelStreaming2
diff --git a/Engine/Plugins/Media/PixelStreaming2/Source/PixelStreaming2RTC/Private/EpicRtcVideoDecoder.h b/Engine/Plugins/Media/PixelStreaming2/Source/PixelStreaming2RTC/Private/EpicRtcVideoDecoder.h
--- a/Engine/Plugins/Media/PixelStreaming2/Source/PixelStreaming2RTC/Private/EpicRtcVideoDecoder.h
+++ b/Engine/Plugins/Media/PixelStreaming2/Source/PixelStreaming2RTC/Private/EpicRtcVideoDecoder.h
@@ -19,49 +19,6 @@
 
 namespace UE::PixelStreaming2
 {
-	template <typename TResolvableVideoResource, typename TVideoResource>
-	class TVideoResourcePool
-	{
-	public:
-		TResolvableVideoResource& GetOrCreate()
-		{
-			ON_SCOPE_EXIT
-			{
-				IPixelStreaming2Stats::Get().GraphValue(TEXT("NumDecodeResource"), Resources.Num(), 1, 0.f, 120.f);
-			};
-
-			if (RHIGetInterfaceType() == ERHIInterfaceType::Vulkan)
-			{
-				// Vulkan has lifetime issues so we always use the same resource
-				if (Resources.Num() == 0)
-				{
-					return Resources.Emplace_GetRef();
-				}
-				return Resources[0];
-			}
-
-			TResolvableVideoResource* Resource = Resources.FindByPredicate([](TResolvableVideoResource& ResolvableResource) {
-				TSharedPtr<TVideoResource>& Resolved = ResolvableResource; // Activates the TSharedPtr operator overload which will resolve the resource
-				if (!Resolved.IsValid())
-				{
-					return false;
-				}
-
-				return !Resolved->IsInUse();
-			});
-
-			if (!Resource)
-			{
-				return Resources.Emplace_GetRef();
-			}
-
-			return *Resource;
-		}
-
-	private:
-		TArray<TResolvableVideoResource> Resources;
-	};
-
 	template <std::derived_from<FVideoResource> TVideoResource>
 	class TEpicRtcVideoDecoder : public EpicRtcVideoDecoderInterface
 	{
@@ -85,20 +42,14 @@ namespace UE::PixelStreaming2
 		uint16_t										   FrameCount;
 
 		// clang-format off
-		using ResourcePoolType = typename std::conditional<std::is_same_v<TVideoResource, FVideoResourceRHI>, TVideoResourcePool<FResolvableVideoResourceRHI, FVideoResourceRHI>, 
-								 typename std::conditional<std::is_same_v<TVideoResource, FVideoResourceCPU>, TVideoResourcePool<FResolvableVideoResourceCPU, FVideoResourceCPU>, 
-								 // We only support RHI and CPU resources so if we create one that isn't this type it will error (on purpose)
-								 void>::type>::type; 
+		using DecoderResourceType = typename std::conditional<std::is_same_v<TVideoResource, FVideoResourceRHI>, FResolvableVideoResourceRHI, 
+								    typename std::conditional<std::is_same_v<TVideoResource, FVideoResourceCPU>, FResolvableVideoResourceCPU, 
+								    // We only support RHI and CPU resources so if we create one that isn't this type it will error (on purpose)
+								    void>::type>::type; 
 
-		// TODO (Eden.Harris) RTCP-7927 Use VideoResources and FetchOrCreate rather than flip flopping
-		// ResourcePoolType VideoResources;
+		DecoderResourceType DecoderResource;
 		// clang-format on
 
-		// TODO (Eden.Harris) RTCP-7927 Use VideoResources and FetchOrCreate rather than flip flopping
-		TStaticArray<FResolvableVideoResourceRHI, 2> VideoResourcesRHI;
-		TStaticArray<FResolvableVideoResourceCPU, 2> VideoResourcesCPU;
-		std::atomic<uint32>							 VideoResourceIndex;
-
 	private:
 		bool LateInitDecoder();
 
@@ -107,4 +58,4 @@ namespace UE::PixelStreaming2
 		EPICRTC_REFCOUNT_INTERFACE_IN_PLACE
 		// End EpicRtcRefCountInterface
 	};
-} // namespace UE::PixelStreaming2
\ No newline at end of file
+} // namespace UE::PixelStreaming2
