[AVCodecs] Disable VTCodecs if IOSurface is unsupported.

diff --git a/Engine/Plugins/Experimental/AVCodecs/VTCodecs/Source/VTCodecs/Private/VT.cpp b/Engine/Plugins/Experimental/AVCodecs/VTCodecs/Source/VTCodecs/Private/VT.cpp
--- a/Engine/Plugins/Experimental/AVCodecs/VTCodecs/Source/VTCodecs/Private/VT.cpp
+++ b/Engine/Plugins/Experimental/AVCodecs/VTCodecs/Source/VTCodecs/Private/VT.cpp
@@ -1,17 +1,50 @@
 // Copyright Epic Games, Inc. All Rights Reserved.
 
 #include "VT.h"
+
+#include "AVResult.h"
 #include "HAL/PlatformProcess.h"
 #include "HAL/Thread.h"
 #include "RHI.h"
 
+THIRD_PARTY_INCLUDES_START
+#include <IOKit/IOKitLib.h>
+THIRD_PARTY_INCLUDES_END
+
 REGISTER_TYPEID(FVT);
 
-FVT::FVT()
+namespace
+{
+	// If IoSurface is unsupported, then VTIsHardwareDecodeSupported will assert internally.	
+	bool IsIoSurfaceAvailable()
+	{
+		// IOSurfaceRoot is determined by running `ioreg -n IOSurfaceRoot` which prints
+		// "CFBundleIdentifier" = "com.apple.iokit.IOSurface"
+		// "IOMatchCategory" = "IOSurfaceRoot"
+		io_service_t Service = IOServiceGetMatchingService(kIOMainPortDefault, IOServiceMatching("IOSurfaceRoot"));
+		if (!Service)
+		{
+			FAVResult::Log(EAVResult::Warning, TEXT("Failed to get IOSurfaceRoot service"));
+			return false;
+		}
+		io_connect_t Connect = IO_OBJECT_NULL;
+		kern_return_t KernResult = IOServiceOpen(Service, mach_task_self(), 0, &Connect);
+		IOObjectRelease(Service);
+		if (KernResult != KERN_SUCCESS)
+		{
+			FAVResult::Log(EAVResult::Warning, FString::Printf(TEXT("IOServiceOpen Failed to open IOSurfaceRoot with: [%d]"), static_cast<int>(KernResult)));
+			return false;
+		}
+		IOServiceClose(Connect);
+		return true;
+	}
+}
+
+FVT::FVT() : bIsIoSurfaceAvailable(IsIoSurfaceAvailable())
 {
 }
 
 bool FVT::IsValid() const
 {
-	return bHasCompatibleGPU;
+	return bHasCompatibleGPU && bIsIoSurfaceAvailable;
 }
diff --git a/Engine/Plugins/Experimental/AVCodecs/VTCodecs/Source/VTCodecs/Private/VTCodecsModule.cpp b/Engine/Plugins/Experimental/AVCodecs/VTCodecs/Source/VTCodecs/Private/VTCodecsModule.cpp
--- a/Engine/Plugins/Experimental/AVCodecs/VTCodecs/Source/VTCodecs/Private/VTCodecsModule.cpp
+++ b/Engine/Plugins/Experimental/AVCodecs/VTCodecs/Source/VTCodecs/Private/VTCodecsModule.cpp
@@ -69,9 +69,9 @@ public:
             ::And<FVideoEncoderConfigVT, FVideoEncoderConfigH264>(
                 [CheckCodecSupport](TSharedRef<FAVDevice> const& NewDevice, TSharedRef<FAVInstance> const& NewInstance)
                 {
-					static bool bSupportsCodec = CheckCodecSupport(NewDevice, TEXT("H.264"));
+					static bool bSupportsCodec = FAPI::Get<FVT>().IsValid() && CheckCodecSupport(NewDevice, TEXT("H.264"));
 
-                    return FAPI::Get<FVT>().IsValid() && bSupportsCodec;
+					return bSupportsCodec;
                 });
 
         FVideoEncoder::RegisterPermutationsOf<TVideoEncoderVT<FVideoResourceMetal>>
@@ -79,45 +79,46 @@ public:
             ::And<FVideoEncoderConfigVT, FVideoEncoderConfigH265>(
                 [CheckCodecSupport](TSharedRef<FAVDevice> const& NewDevice, TSharedRef<FAVInstance> const& NewInstance)
                 {
-					static bool bSupportsCodec = CheckCodecSupport(NewDevice, TEXT("HEVC"));
+					static bool bSupportsCodec = FAPI::Get<FVT>().IsValid() && CheckCodecSupport(NewDevice, TEXT("HEVC"));
 
-                    return FAPI::Get<FVT>().IsValid() && bSupportsCodec;
+					return bSupportsCodec;
                 });
 
         VTRegisterSupplementalVideoDecoderIfAvailable(kCMVideoCodecType_VP9);
 
-        if(VTIsHardwareDecodeSupported(kCMVideoCodecType_H264))
-        {
-            FVideoDecoder::RegisterPermutationsOf<FVideoDecoderVT>
-                ::With<FVideoResourceMetal>
-                ::And<FVideoDecoderConfigVT, FVideoDecoderConfigH264>(
-                    [](TSharedRef<FAVDevice> const& NewDevice, TSharedRef<FAVInstance> const& NewInstance)
-                    {
-                        return FAPI::Get<FVT>().IsValid();
-                    });
-        }
-
-        if(VTIsHardwareDecodeSupported(kCMVideoCodecType_HEVC))
-        {
-            FVideoDecoder::RegisterPermutationsOf<FVideoDecoderVT>
-                ::With<FVideoResourceMetal>
-                ::And<FVideoDecoderConfigVT, FVideoDecoderConfigH265>(
-                    [](TSharedRef<FAVDevice> const& NewDevice, TSharedRef<FAVInstance> const& NewInstance)
-                    {
-                        return FAPI::Get<FVT>().IsValid();
-                    });
-        }
-
-        if(VTIsHardwareDecodeSupported(kCMVideoCodecType_VP9))
-        {
-            FVideoDecoder::RegisterPermutationsOf<FVideoDecoderVT>
-                ::With<FVideoResourceMetal>
-                ::And<FVideoDecoderConfigVT, FVideoDecoderConfigVP9>(
-                    [](TSharedRef<FAVDevice> const& NewDevice, TSharedRef<FAVInstance> const& NewInstance)
-                    {
-                        return FAPI::Get<FVT>().IsValid();
-                    });
-        }
+		FVideoDecoder::RegisterPermutationsOf<FVideoDecoderVT>
+			::With<FVideoResourceMetal>
+			::And<FVideoDecoderConfigVT, FVideoDecoderConfigH264>(
+				[](TSharedRef<FAVDevice> const& NewDevice, TSharedRef<FAVInstance> const& NewInstance)
+				{
+					// IsValid must be called before VTIsHardwareDecodeSupported as it needs IOSurface support
+					static bool bSupportsCodec = FAPI::Get<FVT>().IsValid() && VTIsHardwareDecodeSupported(kCMVideoCodecType_H264);
+					
+					return bSupportsCodec;
+				});
+
+		FVideoDecoder::RegisterPermutationsOf<FVideoDecoderVT>
+			::With<FVideoResourceMetal>
+			::And<FVideoDecoderConfigVT, FVideoDecoderConfigH265>(
+				[](TSharedRef<FAVDevice> const& NewDevice, TSharedRef<FAVInstance> const& NewInstance)
+				{
+					// IsValid must be called before VTIsHardwareDecodeSupported as it needs IOSurface support
+					static bool bSupportsCodec = FAPI::Get<FVT>().IsValid() && VTIsHardwareDecodeSupported(kCMVideoCodecType_HEVC);
+					
+					return bSupportsCodec;
+				});
+
+		FVideoDecoder::RegisterPermutationsOf<FVideoDecoderVT>
+			::With<FVideoResourceMetal>
+			::And<FVideoDecoderConfigVT, FVideoDecoderConfigVP9>(
+				[](TSharedRef<FAVDevice> const& NewDevice, TSharedRef<FAVInstance> const& NewInstance)
+				{
+					// IsValid must be called before VTIsHardwareDecodeSupported as it needs IOSurface support
+					static bool bSupportsCodec = FAPI::Get<FVT>().IsValid() && VTIsHardwareDecodeSupported(kCMVideoCodecType_VP9);
+					
+					return bSupportsCodec;
+				});
+
 	}
 };
 
diff --git a/Engine/Plugins/Experimental/AVCodecs/VTCodecs/Source/VTCodecs/Public/VT.h b/Engine/Plugins/Experimental/AVCodecs/VTCodecs/Source/VTCodecs/Public/VT.h
--- a/Engine/Plugins/Experimental/AVCodecs/VTCodecs/Source/VTCodecs/Public/VT.h
+++ b/Engine/Plugins/Experimental/AVCodecs/VTCodecs/Source/VTCodecs/Public/VT.h
@@ -11,11 +11,13 @@ THIRD_PARTY_INCLUDES_END
 class VTCODECS_API FVT : public FAPI
 {
 public:
-	bool bHasCompatibleGPU = true;
+	bool bHasCompatibleGPU = false;
 
 	FVT();
 
 	virtual bool IsValid() const override;
+private:
+	bool bIsIoSurfaceAvailable;
 };
 
 DECLARE_TYPEID(FVT, VTCODECS_API);
